Initial support for composite fonts for Slate, UMG, and Canvas

Slate now has the concept of composite fonts and font families (via FCompositeFont and FTypeface). A composite font is a font that contains a default font family, as well as any number of sub-font families which should be used for a given set of character ranges. This change will greatly improve the localization support for fonts.

UFont assets can now use two forms of caching "offline" (which is the way they have always worked historically), and "runtime" (which utilizes the Slate font cache to cache glyphs on demand). These runtime cached UFont assets are now the only way to specify which font to use for a UMG widget, and address the previous issues about ensuring that the required font files were staged for your game.

The Slate font cache now works on FFontData instances, rather than filenames. This allows the UFont asset to embed a blob of TTF or OTF font data inside it, rather than require the fonts be loaded from files on disk.

The Canvas text renderer has been updated to support runtime cached UFont assets. This gives our legacy Canvas based tools the same improved font localization support as the rest of the Slate-based editor UI.

UFont asset creation has been changed to use runtime caching by default, and additionally, you can now import a TTF or OTF file via the Content Browser and it will automatically create a UFont asset. If you still want an offline cached UFont asset, you can just change the cache type in the font editor, and the usual font picker dialog will appear and allow you to generate a font texture atlas.

diff --git a/Engine/Shaders/SimpleElementPixelShader.usf b/Engine/Shaders/SimpleElementPixelShader.usf
--- a/Engine/Shaders/SimpleElementPixelShader.usf
+++ b/Engine/Shaders/SimpleElementPixelShader.usf
@@ -81,6 +81,36 @@ void Main(
 #endif
 }
 
+void AlphaOnlyMain(
+	in float2 TextureCoordinate : TEXCOORD0,
+	in float4 Color : TEXCOORD1,
+	in float4 HitProxyId : TEXCOORD2,
+	out float4 OutColor : SV_Target0
+#if FEATURE_LEVEL >= FEATURE_LEVEL_SM4
+	,out float4 OutWorldNormal : SV_Target1 
+	,in float4 SvPosition : SV_Position
+#endif
+	)
+{
+#if FEATURE_LEVEL >= FEATURE_LEVEL_SM4
+	// depth test manually if compositing editor primitives
+	FLATTEN
+	if (bEnableEditorPrimitiveDepthTest)
+	{
+		PerformEditorCompositingDepthTest(SvPosition);
+	}
+#endif
+	
+	float4 TmpColor = Color;
+	TmpColor.a *= Texture2DSample_A8(InTexture, InTextureSampler, TextureCoordinate);
+	OutColor = RETURN_COLOR(TmpColor);
+
+#if FEATURE_LEVEL >= FEATURE_LEVEL_SM4
+	// Set the G buffer bits that indicate unlit
+	OutWorldNormal = 0;
+#endif
+}
+
 #if METAL_PROFILE
 static const half Gamma = 1.0f;
 #else
diff --git a/Engine/Source/Developer/AssetTools/Private/AssetTypeActions/AssetTypeActions_Font.cpp b/Engine/Source/Developer/AssetTools/Private/AssetTypeActions/AssetTypeActions_Font.cpp
--- a/Engine/Source/Developer/AssetTools/Private/AssetTypeActions/AssetTypeActions_Font.cpp
+++ b/Engine/Source/Developer/AssetTools/Private/AssetTypeActions/AssetTypeActions_Font.cpp
@@ -13,7 +13,10 @@ void FAssetTypeActions_Font::GetActions(const TArray<UObject*>& InObjects, FMenu
 		LOCTEXT("ReimportFontLabel", "Reimport"),
 		LOCTEXT("ReimportFontTooltip", "Reimport the selected font(s)."),
 		FSlateIcon(FEditorStyle::GetStyleSetName(), "ContentBrowser.AssetActions.ReimportAsset"),
-		FUIAction(FExecuteAction::CreateSP(this, &FAssetTypeActions_Font::ExecuteReimport, Fonts))
+		FUIAction(
+			FExecuteAction::CreateSP(this, &FAssetTypeActions_Font::ExecuteReimport, Fonts),
+			FCanExecuteAction::CreateSP(this, &FAssetTypeActions_Font::CanExecuteReimport, Fonts)
+			)
 		);
 }
 
@@ -32,6 +35,24 @@ void FAssetTypeActions_Font::OpenAssetEditor( const TArray<UObject*>& InObjects,
 	}
 }
 
+bool FAssetTypeActions_Font::CanExecuteReimport(const TArray<TWeakObjectPtr<UFont>> Objects) const
+{
+	for (auto ObjIt = Objects.CreateConstIterator(); ObjIt; ++ObjIt)
+	{
+		auto Object = (*ObjIt).Get();
+		if (Object)
+		{
+			// We allow a reimport if any of the fonts are using an offline cache
+			if (Object->FontCacheType == EFontCacheType::Offline)
+			{
+				return true;
+			}
+		}
+	}
+
+	return false;
+}
+
 void FAssetTypeActions_Font::ExecuteReimport(const TArray<TWeakObjectPtr<UFont>> Objects) const
 {
 	for (auto ObjIt = Objects.CreateConstIterator(); ObjIt; ++ObjIt)
@@ -39,8 +60,12 @@ void FAssetTypeActions_Font::ExecuteReimport(const TArray<TWeakObjectPtr<UFont>>
 		auto Object = (*ObjIt).Get();
 		if (Object)
 		{
-			// Fonts fail to reimport if they ask for a new file if missing
-			FReimportManager::Instance()->Reimport(Object, /*bAskForNewFileIfMissing=*/false);
+			// Skip fonts that aren't using an offline cache, as they can't be reimported
+			if (Object->FontCacheType != EFontCacheType::Offline)
+			{
+				// Fonts fail to reimport if they ask for a new file if missing
+				FReimportManager::Instance()->Reimport(Object, /*bAskForNewFileIfMissing=*/false);
+			}
 		}
 	}
 }
diff --git a/Engine/Source/Developer/AssetTools/Private/AssetTypeActions/AssetTypeActions_Font.h b/Engine/Source/Developer/AssetTools/Private/AssetTypeActions/AssetTypeActions_Font.h
--- a/Engine/Source/Developer/AssetTools/Private/AssetTypeActions/AssetTypeActions_Font.h
+++ b/Engine/Source/Developer/AssetTools/Private/AssetTypeActions/AssetTypeActions_Font.h
@@ -12,9 +12,12 @@ public:
 	virtual bool HasActions(const TArray<UObject*>& InObjects) const override { return true; }
 	virtual void GetActions(const TArray<UObject*>& InObjects, FMenuBuilder& MenuBuilder) override;
 	virtual void OpenAssetEditor( const TArray<UObject*>& InObjects, TSharedPtr<class IToolkitHost> EditWithinLevelEditor = TSharedPtr<IToolkitHost>() ) override;
-	virtual uint32 GetCategories() override { return EAssetTypeCategories::MaterialsAndTextures; }
+	virtual uint32 GetCategories() override { return EAssetTypeCategories::UI; }
 
 private:
+	/** Can we execute a reimport for the selected objects? */
+	bool CanExecuteReimport(const TArray<TWeakObjectPtr<UFont>> Objects) const;
+
 	/** Handler for when Reimport is selected */
 	void ExecuteReimport(const TArray<TWeakObjectPtr<UFont>> Objects) const;
 };
\ No newline at end of file
diff --git a/Engine/Source/Developer/StandaloneRenderer/Private/OpenGL/SlateOpenGLTextures.cpp b/Engine/Source/Developer/StandaloneRenderer/Private/OpenGL/SlateOpenGLTextures.cpp
--- a/Engine/Source/Developer/StandaloneRenderer/Private/OpenGL/SlateOpenGLTextures.cpp
+++ b/Engine/Source/Developer/StandaloneRenderer/Private/OpenGL/SlateOpenGLTextures.cpp
@@ -66,14 +66,14 @@ void FSlateOpenGLTexture::UpdateTexture(const TArray<uint8>& Bytes)
 
 FSlateFontTextureOpenGL::FSlateFontTextureOpenGL( uint32 Width, uint32 Height )
 	: FSlateFontAtlas( Width, Height ) 
-	, Texture(NULL)
+	, FontTexture(nullptr)
 {
 
 }
 
 FSlateFontTextureOpenGL::~FSlateFontTextureOpenGL()
 {
-	delete Texture;
+	delete FontTexture;
 }
 
 void FSlateFontTextureOpenGL::CreateFontTexture()
@@ -93,8 +93,8 @@ void FSlateFontTextureOpenGL::CreateFontTexture()
 	glTexImage2D( GL_TEXTURE_2D, 0, Format, AtlasWidth, AtlasHeight, 0, Format, GL_UNSIGNED_BYTE, NULL );
 
 	// Create a new slate texture for use in rendering
-	Texture = new FSlateOpenGLTexture( AtlasWidth, AtlasHeight );
-	Texture->Init( TextureID );
+	FontTexture = new FSlateOpenGLTexture( AtlasWidth, AtlasHeight );
+	FontTexture->Init( TextureID );
 
 
 }
@@ -102,12 +102,12 @@ void FSlateFontTextureOpenGL::CreateFontTexture()
 void FSlateFontTextureOpenGL::ConditionalUpdateTexture()
 {
 	// The texture may not be valid when calling this as OpenGL must wait until after the first viewport has been created to create a texture
-	if( bNeedsUpdate && Texture )
+	if( bNeedsUpdate && FontTexture )
 	{
 		check(AtlasData.Num()>0);
 
 		// Completely the texture data each time characters are added
-		glBindTexture(GL_TEXTURE_2D, Texture->GetTypedResource() );
+		glBindTexture(GL_TEXTURE_2D, FontTexture->GetTypedResource() );
 		GLint Format = GL_ALPHA;
 #if PLATFORM_MAC // Make this texture use a DMA'd client storage backing store on OS X, where these extensions always exist
 				 // This avoids a problem on Intel & Nvidia cards that makes characters disappear as well as making the texture updates
diff --git a/Engine/Source/Developer/StandaloneRenderer/Private/OpenGL/SlateOpenGLTextures.h b/Engine/Source/Developer/StandaloneRenderer/Private/OpenGL/SlateOpenGLTextures.h
--- a/Engine/Source/Developer/StandaloneRenderer/Private/OpenGL/SlateOpenGLTextures.h
+++ b/Engine/Source/Developer/StandaloneRenderer/Private/OpenGL/SlateOpenGLTextures.h
@@ -53,8 +53,9 @@ public:
 
 	/** FSlateFontAtlas interface */
 	virtual void ConditionalUpdateTexture();
-	virtual class FSlateShaderResource* GetTexture() { return Texture; }
+	virtual class FSlateShaderResource* GetSlateTexture() override { return FontTexture; }
+	virtual class FTextureResource* GetEngineTexture() override { return nullptr; }
 private:
-	FSlateOpenGLTexture* Texture;
+	FSlateOpenGLTexture* FontTexture;
 };
 
diff --git a/Engine/Source/Developer/StandaloneRenderer/Private/Windows/D3D/SlateD3DTextures.h b/Engine/Source/Developer/StandaloneRenderer/Private/Windows/D3D/SlateD3DTextures.h
--- a/Engine/Source/Developer/StandaloneRenderer/Private/Windows/D3D/SlateD3DTextures.h
+++ b/Engine/Source/Developer/StandaloneRenderer/Private/Windows/D3D/SlateD3DTextures.h
@@ -65,7 +65,8 @@ public:
 
 	/** FSlateFontAtlas interface */
 	virtual void ConditionalUpdateTexture();
-	virtual class FSlateShaderResource* GetTexture() { return FontTexture; }
+	virtual class FSlateShaderResource* GetSlateTexture() override { return FontTexture; }
+	virtual class FTextureResource* GetEngineTexture() override { return nullptr; }
 private:
 	/** Texture used for rendering */
 	FSlateD3DTexture* FontTexture;
diff --git a/Engine/Source/Editor/DetailCustomizations/Private/DetailCustomizations.cpp b/Engine/Source/Editor/DetailCustomizations/Private/DetailCustomizations.cpp
--- a/Engine/Source/Editor/DetailCustomizations/Private/DetailCustomizations.cpp
+++ b/Engine/Source/Editor/DetailCustomizations/Private/DetailCustomizations.cpp
@@ -48,6 +48,7 @@
 #include "Runtime/Engine/Classes/PhysicsEngine/BodySetup.h"
 #include "SlateBrushCustomization.h"
 #include "SlateSoundCustomization.h"
+#include "SlateFontInfoCustomization.h"
 #include "MarginCustomization.h"
 #include "SceneComponentDetails.h"
 #include "PhysicsConstraintComponentDetails.h"
@@ -160,6 +161,7 @@ void FDetailCustomizationsModule::RegisterPropertyTypeCustomizations()
 	RegisterCustomPropertyTypeLayout("VehicleTransmissionData", FOnGetPropertyTypeCustomizationInstance::CreateStatic(&FVehicleTransmissionDataCustomization::MakeInstance));
 	RegisterCustomPropertyTypeLayout("SlateBrush", FOnGetPropertyTypeCustomizationInstance::CreateStatic(&FSlateBrushStructCustomization::MakeInstance, true));
 	RegisterCustomPropertyTypeLayout("SlateSound", FOnGetPropertyTypeCustomizationInstance::CreateStatic(&FSlateSoundStructCustomization::MakeInstance));
+	RegisterCustomPropertyTypeLayout("SlateFontInfo", FOnGetPropertyTypeCustomizationInstance::CreateStatic(&FSlateFontInfoStructCustomization::MakeInstance));
 	RegisterCustomPropertyTypeLayout("Guid", FOnGetPropertyTypeCustomizationInstance::CreateStatic(&FGuidStructCustomization::MakeInstance));
 	RegisterCustomPropertyTypeLayout("Key", FOnGetPropertyTypeCustomizationInstance::CreateStatic(&FKeyStructCustomization::MakeInstance));
 	RegisterCustomPropertyTypeLayout("FloatRange", FOnGetPropertyTypeCustomizationInstance::CreateStatic(&FRangeStructCustomization<float>::MakeInstance));
diff --git a/Engine/Source/Editor/DetailCustomizations/Private/SlateFontInfoCustomization.cpp b/Engine/Source/Editor/DetailCustomizations/Private/SlateFontInfoCustomization.cpp
new file mode 100644
--- /dev/null
+++ b/Engine/Source/Editor/DetailCustomizations/Private/SlateFontInfoCustomization.cpp
@@ -0,0 +1,203 @@
+// Copyright 1998-2014 Epic Games, Inc. All Rights Reserved.
+
+#include "DetailCustomizationsPrivatePCH.h"
+#include "SlateFontInfoCustomization.h"
+
+#define LOCTEXT_NAMESPACE "SlateFontInfo"
+
+TSharedRef<IPropertyTypeCustomization> FSlateFontInfoStructCustomization::MakeInstance() 
+{
+	return MakeShareable(new FSlateFontInfoStructCustomization());
+}
+
+void FSlateFontInfoStructCustomization::CustomizeHeader(TSharedRef<IPropertyHandle> StructPropertyHandle, FDetailWidgetRow& HeaderRow, IPropertyTypeCustomizationUtils& StructCustomizationUtils)
+{
+	FontObjectProperty = StructPropertyHandle->GetChildHandle(TEXT("FontObject"));
+	check(FontObjectProperty.IsValid());
+
+	TypefaceFontNameProperty = StructPropertyHandle->GetChildHandle(TEXT("TypefaceFontName"));
+	check(TypefaceFontNameProperty.IsValid());
+
+	FontSizeProperty = StructPropertyHandle->GetChildHandle(TEXT("Size"));
+	check(FontSizeProperty.IsValid());
+
+	TArray<void*> StructPtrs;
+	StructPropertyHandle->AccessRawData(StructPtrs);
+	for(auto It = StructPtrs.CreateConstIterator(); It; ++It)
+	{
+		FSlateFontInfo* const SlateFontInfo = reinterpret_cast<FSlateFontInfo*>(*It);
+		SlateFontInfoStructs.Add(SlateFontInfo);
+	}
+
+	HeaderRow
+	.NameContent()
+	[
+		StructPropertyHandle->CreatePropertyNameWidget()
+	]
+	.ValueContent()
+	.MinDesiredWidth(250.0f)
+	.MaxDesiredWidth(0.0f)
+	.HAlign(HAlign_Fill)
+	.VAlign(VAlign_Center)
+	[
+		SNew(SHorizontalBox)
+
+		+SHorizontalBox::Slot()
+		.FillWidth(2)
+		[
+			SNew(SObjectPropertyEntryBox)
+			.PropertyHandle(FontObjectProperty)
+			.AllowedClass(UFont::StaticClass())
+			.OnShouldFilterAsset(FOnShouldFilterAsset::CreateStatic(&FSlateFontInfoStructCustomization::OnFilterFontAsset))
+			.OnObjectChanged(this, &FSlateFontInfoStructCustomization::OnFontChanged)
+			.DisplayUseSelected(false)
+			.DisplayBrowse(false)
+		]
+
+		+SHorizontalBox::Slot()
+		.FillWidth(2)
+		.VAlign(VAlign_Center)
+		[
+			SAssignNew(FontEntryCombo, SComboBox<TSharedPtr<FName>>)
+			.OptionsSource(&FontEntryComboData)
+			.IsEnabled(this, &FSlateFontInfoStructCustomization::IsFontEntryComboEnabled)
+			.OnComboBoxOpening(this, &FSlateFontInfoStructCustomization::OnFontEntryComboOpening)
+			.OnSelectionChanged(this, &FSlateFontInfoStructCustomization::OnFontEntrySelectionChanged)
+			.OnGenerateWidget(this, &FSlateFontInfoStructCustomization::MakeFontEntryWidget)
+			[
+				SNew(STextBlock)
+				.Text(this, &FSlateFontInfoStructCustomization::GetFontEntryComboText)
+				.Font(FEditorStyle::GetFontStyle(TEXT("PropertyWindow.NormalFont")))
+			]
+		]
+
+		+SHorizontalBox::Slot()
+		.FillWidth(1)
+		.VAlign(VAlign_Center)
+		.Padding(FMargin(4.0f, 0.0f, 0.0f, 0.0f))
+		[
+			SNew(SProperty, FontSizeProperty)
+			.ShouldDisplayName(false)
+		]
+	];
+}
+
+void FSlateFontInfoStructCustomization::CustomizeChildren(TSharedRef<IPropertyHandle> StructPropertyHandle, IDetailChildrenBuilder& StructBuilder, IPropertyTypeCustomizationUtils& StructCustomizationUtils)
+{
+}
+
+bool FSlateFontInfoStructCustomization::OnFilterFontAsset(const FAssetData& InAssetData)
+{
+	// We want to filter font assets that aren't valid to use with Slate/UMG
+	return Cast<const UFont>(InAssetData.GetAsset())->FontCacheType != EFontCacheType::Runtime;
+}
+
+void FSlateFontInfoStructCustomization::OnFontChanged(const FAssetData& InAssetData)
+{
+	const UFont* const FontAsset = Cast<const UFont>(InAssetData.GetAsset());
+	const FName FirstFontName = (FontAsset && FontAsset->CompositeFont.DefaultTypeface.Fonts.Num()) ? FontAsset->CompositeFont.DefaultTypeface.Fonts[0].Name : NAME_None;
+
+	for(FSlateFontInfo* FontInfo : SlateFontInfoStructs)
+	{
+		// The font has been updated in the editor, so clear the non-UObject pointer so that the two don't conflict
+		FontInfo->CompositeFont.Reset();
+
+		// We've changed (or cleared) the font asset, so make sure and update the typeface entry name being used by the font info
+		TypefaceFontNameProperty->SetValue(FirstFontName);
+	}
+}
+
+bool FSlateFontInfoStructCustomization::IsFontEntryComboEnabled() const
+{
+	const FSlateFontInfo* const FirstSlateFontInfo = SlateFontInfoStructs[0];
+	const UFont* const FontObject = Cast<const UFont>(FirstSlateFontInfo->FontObject);
+	if(!FontObject)
+	{
+		return false;
+	}
+
+	// Only let people pick an entry if every struct being edited is using the same font object
+	for(int32 FontInfoIndex = 1; FontInfoIndex < SlateFontInfoStructs.Num(); ++FontInfoIndex)
+	{
+		const FSlateFontInfo* const OtherSlateFontInfo = SlateFontInfoStructs[FontInfoIndex];
+		const UFont* const OtherFontObject = Cast<const UFont>(OtherSlateFontInfo->FontObject);
+		if(FontObject != OtherFontObject)
+		{
+			return false;
+		}
+	}
+
+	return true;
+}
+
+void FSlateFontInfoStructCustomization::OnFontEntryComboOpening()
+{
+	const FSlateFontInfo* const FirstSlateFontInfo = SlateFontInfoStructs[0];
+	const UFont* const FontObject = Cast<const UFont>(FirstSlateFontInfo->FontObject);
+	check(FontObject);
+
+	const FName ActiveFontEntry = GetActiveFontEntry();
+	TSharedPtr<FName> SelectedNamePtr;
+
+	FontEntryComboData.Empty();
+	for(const FTypefaceEntry& TypefaceEntry : FontObject->CompositeFont.DefaultTypeface.Fonts)
+	{
+		TSharedPtr<FName> NameEntryPtr = MakeShareable(new FName(TypefaceEntry.Name));
+		FontEntryComboData.Add(NameEntryPtr);
+
+		if(!TypefaceEntry.Name.IsNone() && TypefaceEntry.Name == ActiveFontEntry)
+		{
+			SelectedNamePtr = NameEntryPtr;
+		}
+	}
+
+	FontEntryComboData.Sort([](const TSharedPtr<FName>& One, const TSharedPtr<FName>& Two) -> bool
+	{
+		return One->ToString() < Two->ToString();
+	});
+
+	FontEntryCombo->ClearSelection();
+	FontEntryCombo->RefreshOptions();
+	FontEntryCombo->SetSelectedItem(SelectedNamePtr);
+}
+
+void FSlateFontInfoStructCustomization::OnFontEntrySelectionChanged(TSharedPtr<FName> InNewSelection, ESelectInfo::Type)
+{
+	if(InNewSelection.IsValid())
+	{
+		const FSlateFontInfo* const FirstSlateFontInfo = SlateFontInfoStructs[0];
+		if(FirstSlateFontInfo->TypefaceFontName != *InNewSelection)
+		{
+			TypefaceFontNameProperty->SetValue(*InNewSelection);
+		}
+	}
+}
+
+TSharedRef<SWidget> FSlateFontInfoStructCustomization::MakeFontEntryWidget(TSharedPtr<FName> InFontEntry)
+{
+	return
+		SNew(STextBlock)
+		.Text(FText::FromName(*InFontEntry))
+		.Font(FEditorStyle::GetFontStyle(TEXT("PropertyWindow.NormalFont")));
+}
+
+FText FSlateFontInfoStructCustomization::GetFontEntryComboText() const
+{
+	return FText::FromName(GetActiveFontEntry());
+}
+
+FName FSlateFontInfoStructCustomization::GetActiveFontEntry() const
+{
+	const FSlateFontInfo* const FirstSlateFontInfo = SlateFontInfoStructs[0];
+	const UFont* const FontObject = Cast<const UFont>(FirstSlateFontInfo->FontObject);
+	if(FontObject)
+	{
+		return (FirstSlateFontInfo->TypefaceFontName.IsNone() && FontObject->CompositeFont.DefaultTypeface.Fonts.Num())
+			? FontObject->CompositeFont.DefaultTypeface.Fonts[0].Name
+			: FirstSlateFontInfo->TypefaceFontName;
+	}
+
+	return NAME_None;
+}
+
+#undef LOCTEXT_NAMESPACE
diff --git a/Engine/Source/Editor/DetailCustomizations/Public/Customizations/SlateFontInfoCustomization.h b/Engine/Source/Editor/DetailCustomizations/Public/Customizations/SlateFontInfoCustomization.h
new file mode 100644
--- /dev/null
+++ b/Engine/Source/Editor/DetailCustomizations/Public/Customizations/SlateFontInfoCustomization.h
@@ -0,0 +1,57 @@
+// Copyright 1998-2014 Epic Games, Inc. All Rights Reserved.
+
+#pragma once
+
+/** Customize the appearance of an FSlateFontInfo */
+class DETAILCUSTOMIZATIONS_API FSlateFontInfoStructCustomization : public IPropertyTypeCustomization
+{
+public:
+	static TSharedRef<IPropertyTypeCustomization> MakeInstance();
+
+	/** IPropertyTypeCustomization interface */
+	virtual void CustomizeHeader(TSharedRef<class IPropertyHandle> StructPropertyHandle, class FDetailWidgetRow& HeaderRow, IPropertyTypeCustomizationUtils& StructCustomizationUtils) override;
+	virtual void CustomizeChildren(TSharedRef<class IPropertyHandle> StructPropertyHandle, class IDetailChildrenBuilder& StructBuilder, IPropertyTypeCustomizationUtils& StructCustomizationUtils) override;
+
+protected:
+	/** Called to filter out invalid font assets */
+	static bool OnFilterFontAsset(const FAssetData& InAssetData);
+
+	/** Called when the font object used by this FSlateFontInfo has been changed */
+	void OnFontChanged(const FAssetData& InAssetData);
+
+	/** Called to see whether the font entry combo should be enabled */
+	bool IsFontEntryComboEnabled() const;
+
+	/** Called before the font entry combo is opened - used to update the list of available font entries */
+	void OnFontEntryComboOpening();
+
+	/** Called when the selection of the font entry combo is changed */
+	void OnFontEntrySelectionChanged(TSharedPtr<FName> InNewSelection, ESelectInfo::Type);
+
+	/** Make the widget for an entry in the font entry combo */
+	TSharedRef<SWidget> MakeFontEntryWidget(TSharedPtr<FName> InFontEntry);
+
+	/** Get the text to use for the font entry combo button */
+	FText GetFontEntryComboText() const;
+
+	/** Get the name of the currently active font entry (may not be the selected entry if the entry is set to use "None") */
+	FName GetActiveFontEntry() const;
+
+	/** Array of FSlateFontInfo instances this customization is currently editing */
+	TArray<FSlateFontInfo*> SlateFontInfoStructs;
+
+	/** Handle to the "FontObject" property being edited */
+	TSharedPtr<IPropertyHandle> FontObjectProperty;
+
+	/** Handle to the "TypefaceFontName" property being edited */
+	TSharedPtr<IPropertyHandle> TypefaceFontNameProperty;
+
+	/** Handle to the "Size" property being edited */
+	TSharedPtr<IPropertyHandle> FontSizeProperty;
+
+	/** Font entry combo box widget */
+	TSharedPtr<SComboBox<TSharedPtr<FName>>> FontEntryCombo;
+
+	/** Source data for the font entry combo widget */
+	TArray<TSharedPtr<FName>> FontEntryComboData;
+};
diff --git a/Engine/Source/Editor/EditorStyle/Private/SlateEditorStyle.cpp b/Engine/Source/Editor/EditorStyle/Private/SlateEditorStyle.cpp
--- a/Engine/Source/Editor/EditorStyle/Private/SlateEditorStyle.cpp
+++ b/Engine/Source/Editor/EditorStyle/Private/SlateEditorStyle.cpp
@@ -5342,6 +5342,9 @@ void FSlateEditorStyle::FStyle::SetupToolkitStyles()
 		Set( "FontEditor.FontBackgroundColor.Small", new IMAGE_BRUSH( "Icons/icon_FontEd_Background_40x", Icon20x20 ) );
 		Set( "FontEditor.FontForegroundColor", new IMAGE_BRUSH( "Icons/icon_FontEd_Foreground_40x", Icon40x40 ) );
 		Set( "FontEditor.FontForegroundColor.Small", new IMAGE_BRUSH( "Icons/icon_FontEd_Foreground_40x", Icon20x20 ) );
+
+		Set( "FontEditor.Button_Add", new IMAGE_BRUSH( "Icons/PlusSymbol_12x", Icon12x12 ) );
+		Set( "FontEditor.Button_Delete", new IMAGE_BRUSH("Icons/Cross_12x", Icon12x12 ) );
 	}
 
 	// SoundCueGraph Editor
diff --git a/Engine/Source/Editor/FontEditor/FontEditor.Build.cs b/Engine/Source/Editor/FontEditor/FontEditor.Build.cs
--- a/Engine/Source/Editor/FontEditor/FontEditor.Build.cs
+++ b/Engine/Source/Editor/FontEditor/FontEditor.Build.cs
@@ -12,6 +12,7 @@ public class FontEditor : ModuleRules
 				"Core",
 				"CoreUObject",
 				"DesktopPlatform",
+				"DesktopWidgets",
 				"Engine",
                 "InputCore",
 				"Slate",
diff --git a/Engine/Source/Editor/FontEditor/Private/FontEditor.cpp b/Engine/Source/Editor/FontEditor/Private/FontEditor.cpp
--- a/Engine/Source/Editor/FontEditor/Private/FontEditor.cpp
+++ b/Engine/Source/Editor/FontEditor/Private/FontEditor.cpp
@@ -5,6 +5,7 @@
 #include "Toolkits/IToolkitHost.h"
 #include "SColorPicker.h"
 #include "SFontEditorViewport.h"
+#include "SCompositeFontEditor.h"
 #include "FontEditor.h"
 #include "Editor/WorkspaceMenuStructure/Public/WorkspaceMenuStructureModule.h"
 #include "MainFrame.h"
@@ -19,7 +20,8 @@ DEFINE_LOG_CATEGORY_STATIC(LogFontEditor, Log, All);
 
 FString FFontEditor::LastPath;
 
-const FName FFontEditor::ViewportTabId( TEXT( "FontEditor_FontViewport" ) );
+const FName FFontEditor::TexturePagesViewportTabId( TEXT( "FontEditor_TexturePagesViewport" ) );
+const FName FFontEditor::CompositeFontEditorTabId( TEXT( "FontEditor_CompositeFontEditor" ) );
 const FName FFontEditor::PreviewTabId( TEXT( "FontEditor_FontPreview" ) );
 const FName FFontEditor::PropertiesTabId( TEXT( "FontEditor_FontProperties" ) );
 const FName FFontEditor::PagePropertiesTabId( TEXT( "FontEditor_FontPageProperties" ) );
@@ -77,10 +79,17 @@ void FFontEditor::RegisterTabSpawners(const TSharedRef<class FTabManager>& TabMa
 
 	FAssetEditorToolkit::RegisterTabSpawners(TabManager);
 
-	TabManager->RegisterTabSpawner( ViewportTabId,		FOnSpawnTab::CreateSP(this, &FFontEditor::SpawnTab_Viewport) )
-		.SetDisplayName( LOCTEXT("ViewportTab", "Viewport") )
+	TabManager->RegisterTabSpawner( TexturePagesViewportTabId, FOnSpawnTab::CreateSP(this, &FFontEditor::SpawnTab_TexturePagesViewport) )
+		.SetDisplayName( LOCTEXT("TexturePagesViewportTab", "Texture Pages") )
 		.SetGroup(WorkspaceMenuCategoryRef)
-		.SetIcon(FSlateIcon(FEditorStyle::GetStyleSetName(), "LevelEditor.Tabs.Viewports"));
+		.SetIcon(FSlateIcon(FEditorStyle::GetStyleSetName(), "LevelEditor.Tabs.Viewports"))
+		.SetMenuType( TAttribute<ETabSpawnerMenuType::Type>::Create(TAttribute<ETabSpawnerMenuType::Type>::FGetter::CreateSP(this, &FFontEditor::GetTabSpawnerMenuType, TexturePagesViewportTabId)) );
+
+	TabManager->RegisterTabSpawner( CompositeFontEditorTabId, FOnSpawnTab::CreateSP(this, &FFontEditor::SpawnTab_CompositeFontEditor) )
+		.SetDisplayName( LOCTEXT("CompositeFontEditorTab", "Composite Font") )
+		.SetGroup(WorkspaceMenuCategoryRef)
+		.SetIcon(FSlateIcon(FEditorStyle::GetStyleSetName(), "FontEditor.Tabs.PageProperties"))
+		.SetMenuType( TAttribute<ETabSpawnerMenuType::Type>::Create(TAttribute<ETabSpawnerMenuType::Type>::FGetter::CreateSP(this, &FFontEditor::GetTabSpawnerMenuType, CompositeFontEditorTabId)) );
 
 	TabManager->RegisterTabSpawner( PreviewTabId,		FOnSpawnTab::CreateSP(this, &FFontEditor::SpawnTab_Preview) )
 		.SetDisplayName( LOCTEXT("PreviewTab", "Preview") )
@@ -95,14 +104,16 @@ void FFontEditor::RegisterTabSpawners(const TSharedRef<class FTabManager>& TabMa
 	TabManager->RegisterTabSpawner( PagePropertiesTabId,FOnSpawnTab::CreateSP(this, &FFontEditor::SpawnTab_PageProperties) )
 		.SetDisplayName( LOCTEXT("PagePropertiesTab", "Page Details") )
 		.SetGroup(WorkspaceMenuCategoryRef)
-		.SetIcon(FSlateIcon(FEditorStyle::GetStyleSetName(), "FontEditor.Tabs.PageProperties"));
+		.SetIcon(FSlateIcon(FEditorStyle::GetStyleSetName(), "FontEditor.Tabs.PageProperties"))
+		.SetMenuType( TAttribute<ETabSpawnerMenuType::Type>::Create(TAttribute<ETabSpawnerMenuType::Type>::FGetter::CreateSP(this, &FFontEditor::GetTabSpawnerMenuType, PagePropertiesTabId)) );
 }
 
 void FFontEditor::UnregisterTabSpawners(const TSharedRef<class FTabManager>& TabManager)
 {
 	FAssetEditorToolkit::UnregisterTabSpawners(TabManager);
 
-	TabManager->UnregisterTabSpawner( ViewportTabId );	
+	TabManager->UnregisterTabSpawner( TexturePagesViewportTabId );
+	TabManager->UnregisterTabSpawner( CompositeFontEditorTabId );
 	TabManager->UnregisterTabSpawner( PreviewTabId );	
 	TabManager->UnregisterTabSpawner( PropertiesTabId );
 	TabManager->UnregisterTabSpawner( PagePropertiesTabId );
@@ -154,7 +165,7 @@ void FFontEditor::InitFontEditor(const EToolkitMode::Type Mode, const TSharedPtr
 
 	CreateInternalWidgets();
 
-	const TSharedRef<FTabManager::FLayout> StandaloneDefaultLayout = FTabManager::NewLayout("Standalone_FontEditor_Layout_v2")
+	const TSharedRef<FTabManager::FLayout> StandaloneDefaultLayout = FTabManager::NewLayout("Standalone_FontEditor_Layout_v3")
 	->AddArea
 	(
 		FTabManager::NewPrimaryArea() ->SetOrientation( Orient_Vertical )
@@ -172,7 +183,8 @@ void FFontEditor::InitFontEditor(const EToolkitMode::Type Mode, const TSharedPtr
 				->Split
 				(
 					FTabManager::NewStack() ->SetSizeCoefficient(0.85f)
-					->AddTab( ViewportTabId, ETabState::OpenedTab ) ->SetHideTabWell( true )
+					->AddTab( TexturePagesViewportTabId, ETabState::OpenedTab )
+					->AddTab( CompositeFontEditorTabId, ETabState::OpenedTab )
 				)
 				->Split
 				(
@@ -207,11 +219,14 @@ void FFontEditor::InitFontEditor(const EToolkitMode::Type Mode, const TSharedPtr
 	ExtendToolbar();
 	RegenerateMenusAndToolbars();
 
+	UpdateLayout();
+
 	// @todo toolkit world centric editing
 	/*if(IsWorldCentricAssetEditor())
 	{
 		SpawnToolkitTab(GetToolbarTabId(), FString(), EToolkitTabSpot::ToolBar);
-		SpawnToolkitTab(ViewportTabId, FString(), EToolkitTabSpot::Viewport);
+		SpawnToolkitTab(TexturePagesViewportTabId, FString(), EToolkitTabSpot::Viewport);
+		SpawnToolkitTab(CompositeFontEditorTabId, FString(), EToolkitTabSpot::Viewport);
 		SpawnToolkitTab(PreviewTabId, FString(), EToolkitTabSpot::Viewport);
 		SpawnToolkitTab(PropertiesTabId, FString(), EToolkitTabSpot::Details);
 		SpawnToolkitTab(PagePropertiesTabId, FString(), EToolkitTabSpot::Details);
@@ -253,12 +268,12 @@ FLinearColor FFontEditor::GetWorldCentricTabColorScale() const
 	return FLinearColor(0.3f, 0.2f, 0.5f, 0.5f);
 }
 
-TSharedRef<SDockTab> FFontEditor::SpawnTab_Viewport( const FSpawnTabArgs& Args )
+TSharedRef<SDockTab> FFontEditor::SpawnTab_TexturePagesViewport( const FSpawnTabArgs& Args )
 {
-	check( Args.GetTabId().TabType == ViewportTabId );
+	check( Args.GetTabId().TabType == TexturePagesViewportTabId );
 
 	TSharedRef<SDockTab> SpawnedTab = SNew(SDockTab)
-		.Label(LOCTEXT("FontViewportTitle", "Viewport"))
+		.Label(LOCTEXT("TexturePagesViewportTitle", "Texture Pages"))
 		[
 			FontViewport.ToSharedRef()
 		];
@@ -268,6 +283,21 @@ TSharedRef<SDockTab> FFontEditor::SpawnTab_Viewport( const FSpawnTabArgs& Args )
 	return SpawnedTab;
 }
 
+TSharedRef<SDockTab> FFontEditor::SpawnTab_CompositeFontEditor( const FSpawnTabArgs& Args )
+{
+	check( Args.GetTabId().TabType == CompositeFontEditorTabId );
+
+	TSharedRef<SDockTab> SpawnedTab = SNew(SDockTab)
+		.Label(LOCTEXT("CompositeFontEditorTitle", "Composite Font"))
+		[
+			CompositeFontEditor.ToSharedRef()
+		];
+
+	AddToSpawnedToolPanels( Args.GetTabId().TabType, SpawnedTab );
+
+	return SpawnedTab;
+}
+
 TSharedRef<SDockTab> FFontEditor::SpawnTab_Preview( const FSpawnTabArgs& Args )
 {
 	check( Args.GetTabId().TabType == PreviewTabId );
@@ -344,21 +374,135 @@ void FFontEditor::OnPreviewTextChanged(const FText& Text)
 
 void FFontEditor::PostUndo(bool bSuccess)
 {
+	// Make sure we're using the correct layout, as the undo/redo may have changed the font cache type property
+	UpdateLayout();
+
+	CompositeFontEditor->Refresh();
 	FontPreviewWidget->RefreshViewport();
 }
 
 void FFontEditor::NotifyPostChange( const FPropertyChangedEvent& PropertyChangedEvent, class FEditPropertyChain* PropertyThatChanged)
 {
-	FontViewport->RefreshViewport();
+	static const FName FontCacheTypePropertyName("FontCacheType");
+	static const FName CompositeFontPropertyName("CompositeFont");
+
+	if(PropertyChangedEvent.Property && PropertyChangedEvent.Property->GetFName() == FontCacheTypePropertyName)
+	{
+		// Show a warning message, as what we're about to do will destroy any existing data in this font object
+		const EAppReturnType::Type DlgResult = OpenMsgDlgInt(
+			EAppMsgType::YesNo, 
+			LOCTEXT("ChangeCacheTypeWarningMsg", "Changing the cache type will cause this font to be reinitialized (discarding any existing data).\n\nAre you sure you want to proceed?"), 
+			LOCTEXT("ChangeCacheTypeWarningTitle", "Really change the font cache type?")
+			);
+
+		bool bSuccessfullyChangedCacheType = false;
+		if(DlgResult == EAppReturnType::Yes)
+		{
+			bSuccessfullyChangedCacheType = RecreateFontObject(Font->FontCacheType);
+		}
+
+		if(bSuccessfullyChangedCacheType)
+		{
+			CompositeFontEditor->Refresh();
+
+			// If we changed the font cache type, then we need to update the UI to hide the invalid tabs and spawn the new ones
+			UpdateLayout();
+		}
+		else
+		{
+			// Restore the old font cache type
+			switch(Font->FontCacheType)
+			{
+			case EFontCacheType::Offline:
+				Font->FontCacheType = EFontCacheType::Runtime;
+				break;
+
+			case EFontCacheType::Runtime:
+				Font->FontCacheType = EFontCacheType::Offline;
+				break;
+
+			default:
+				break;
+			}
+		}
+	}
+
+	// If we changed a property of the composite font, we need to refresh the composite font editor
+	if(PropertyThatChanged && PropertyThatChanged->GetHead()->GetValue()->GetFName() == CompositeFontPropertyName)
+	{
+		CompositeFontEditor->Refresh();
+	}
+
+	if(Font->FontCacheType == EFontCacheType::Offline)
+	{
+		FontViewport->RefreshViewport();
+	}
+
 	FontPreviewWidget->RefreshViewport();
 }
 
+void FFontEditor::UpdateLayout()
+{
+	if(CurrentEditorLayout.IsSet() && CurrentEditorLayout.GetValue() == Font->FontCacheType)
+	{
+		return;
+	}
+
+	auto CloseTab = [this](const FName& TabName)
+	{
+		TWeakPtr<SDockTab>* const FoundExistingTab = SpawnedToolPanels.Find(TabName);
+		if(FoundExistingTab)
+		{
+			TSharedPtr<SDockTab> ExistingTab = FoundExistingTab->Pin();
+			if(ExistingTab.IsValid())
+			{
+				ExistingTab->RequestCloseTab();
+			}
+		}
+	};
+
+	switch(Font->FontCacheType)
+	{
+	case EFontCacheType::Offline:
+		TabManager->InvokeTab(TexturePagesViewportTabId);
+		TabManager->InvokeTab(PagePropertiesTabId);
+		CloseTab(CompositeFontEditorTabId);
+		break;
+
+	case EFontCacheType::Runtime:
+		TabManager->InvokeTab(CompositeFontEditorTabId);
+		CloseTab(TexturePagesViewportTabId);
+		CloseTab(PagePropertiesTabId);
+		break;
+
+	default:
+		break;
+	}
+
+	CurrentEditorLayout = Font->FontCacheType;
+}
+
+ETabSpawnerMenuType::Type FFontEditor::GetTabSpawnerMenuType( FName InTabName ) const
+{
+	if( (Font->FontCacheType == EFontCacheType::Offline && (InTabName == CompositeFontEditorTabId)) ||
+		(Font->FontCacheType == EFontCacheType::Runtime && (InTabName == TexturePagesViewportTabId || InTabName == PagePropertiesTabId)))
+	{
+		return ETabSpawnerMenuType::Hidden;
+	}
+
+	return ETabSpawnerMenuType::Enabled;
+}
+
 void FFontEditor::CreateInternalWidgets()
 {
 	FontViewport = 
 	SNew(SFontEditorViewport)
 	.FontEditor(SharedThis(this));
 
+	CompositeFontEditor = 
+	SNew(SCompositeFontEditor)
+	.FontEditor(SharedThis(this));
+
 	FontPreview =
 	SNew(SVerticalBox)
 	+SVerticalBox::Slot()
@@ -440,7 +584,8 @@ void FFontEditor::BindCommands()
 
 	ToolkitCommands->MapAction(
 		Commands.UpdateAll,
-		FExecuteAction::CreateSP(this, &FFontEditor::OnUpdateAll));
+		FExecuteAction::CreateSP(this, &FFontEditor::OnUpdateAll),
+		FCanExecuteAction::CreateSP(this, &FFontEditor::OnUpdateAllEnabled));
 
 	ToolkitCommands->MapAction(
 		Commands.ExportPage,
@@ -449,7 +594,8 @@ void FFontEditor::BindCommands()
 
 	ToolkitCommands->MapAction(
 		Commands.ExportAllPages,
-		FExecuteAction::CreateSP(this, &FFontEditor::OnExportAll));
+		FExecuteAction::CreateSP(this, &FFontEditor::OnExportAll),
+		FCanExecuteAction::CreateSP(this, &FFontEditor::OnExportAllEnabled));
 
 	ToolkitCommands->MapAction(
 		Commands.FontBackgroundColor,
@@ -518,7 +664,7 @@ void FFontEditor::OnUpdate()
 
 bool FFontEditor::OnUpdateEnabled() const
 {
-	return FontViewport->GetCurrentSelectedPage() != INDEX_NONE;
+	return Font->FontCacheType == EFontCacheType::Offline && FontViewport->GetCurrentSelectedPage() != INDEX_NONE;
 }
 
 void FFontEditor::OnUpdateAll()
@@ -578,6 +724,11 @@ void FFontEditor::OnUpdateAll()
 	FontPreviewWidget->RefreshViewport();
 }
 
+bool FFontEditor::OnUpdateAllEnabled() const
+{
+	return Font->FontCacheType == EFontCacheType::Offline;
+}
+
 void FFontEditor::OnExport()
 {
 	int32 CurrentSelectedPage = FontViewport->GetCurrentSelectedPage();
@@ -621,7 +772,7 @@ void FFontEditor::OnExport()
 
 bool FFontEditor::OnExportEnabled() const
 {
-	return FontViewport->GetCurrentSelectedPage() != INDEX_NONE;
+	return Font->FontCacheType == EFontCacheType::Offline && FontViewport->GetCurrentSelectedPage() != INDEX_NONE;
 }
 
 void FFontEditor::OnExportAll()
@@ -664,6 +815,11 @@ void FFontEditor::OnExportAll()
 	}
 }
 
+bool FFontEditor::OnExportAllEnabled() const
+{
+	return Font->FontCacheType == EFontCacheType::Offline;
+}
+
 void FFontEditor::OnBackgroundColor()
 {
 	FColor Color = FontPreviewWidget->GetPreviewBackgroundColor();
@@ -798,9 +954,53 @@ void FFontEditor::OnObjectReimported(UObject* InObject)
 	}
 }
 
+bool FFontEditor::RecreateFontObject(const EFontCacheType NewCacheType)
+{
+	bool bSuccess = false;
+
+	UFactory* FontFactoryPtr = nullptr;
+
+	switch(NewCacheType)
+	{
+	case EFontCacheType::Offline:
+		// UTrueTypeFontFactory will create a new font object using a texture generated from a user-selection font
+		FontFactoryPtr = ConstructObject<UTrueTypeFontFactory>(UTrueTypeFontFactory::StaticClass());
+		break;
+
+	case EFontCacheType::Runtime:
+		// UFontFactory will create an empty font ready to add new font files to
+		FontFactoryPtr = ConstructObject<UFontFactory>(UFontFactory::StaticClass());
+		break;
+
+	default:
+		break;
+	}
+
+	if(FontFactoryPtr && FontFactoryPtr->ConfigureProperties())
+	{
+		bSuccess = UFactory::StaticImportObject(Font->GetClass(), Font->GetOuter(), *Font->GetName(), RF_Public|RF_Standalone, TEXT(""), nullptr, FontFactoryPtr) != nullptr;
+	}
+
+	if(bSuccess)
+	{
+		Font->PostEditChange();
+		GEditor->BroadcastObjectReimported(Font);
+	}
+
+	// Let listeners know whether the reimport was successful or not
+	FReimportManager::Instance()->OnPostReimport().Broadcast(Font, bSuccess);
+
+	return bSuccess;
+}
+
 bool FFontEditor::ShouldPromptForNewFilesOnReload(const UObject& EditingObject) const
 {
 	return false;
 }
 
+void FFontEditor::RefreshPreview()
+{
+	FontPreviewWidget->RefreshViewport();
+}
+
 #undef LOCTEXT_NAMESPACE
diff --git a/Engine/Source/Editor/FontEditor/Private/FontEditor.h b/Engine/Source/Editor/FontEditor/Private/FontEditor.h
--- a/Engine/Source/Editor/FontEditor/Private/FontEditor.h
+++ b/Engine/Source/Editor/FontEditor/Private/FontEditor.h
@@ -24,6 +24,7 @@ public:
 	/** IFontEditor interface */
 	virtual UFont* GetFont() const override;
 	virtual void SetSelectedPage(int32 PageIdx) override;
+	virtual void RefreshPreview() override;
 	
 	/** IToolkit interface */
 	virtual FName GetToolkitFName() const override;
@@ -51,6 +52,12 @@ private:
 	/** FNotifyHook interface */
 	virtual void NotifyPostChange( const FPropertyChangedEvent& PropertyChangedEvent, class FEditPropertyChain* PropertyThatChanged) override;
 
+	/** Update the font editor UI based on the type of font being edited */
+	void UpdateLayout();
+
+	/** Get the menu type to use for the given tab spawner */
+	ETabSpawnerMenuType::Type GetTabSpawnerMenuType( FName InTabName ) const;
+
 	/** Creates all internal widgets for the tabs to point at */
 	void CreateInternalWidgets();
 
@@ -64,9 +71,11 @@ private:
 	void OnUpdate();
 	bool OnUpdateEnabled() const;
 	void OnUpdateAll();
+	bool OnUpdateAllEnabled() const;
 	void OnExport();
 	bool OnExportEnabled() const;
 	void OnExportAll();
+	bool OnExportAllEnabled() const;
 	void OnBackgroundColor();
 	bool OnBackgroundColorEnabled() const;
 	void OnForegroundColor();
@@ -76,8 +85,11 @@ private:
 	/** Common method for replacing a font page with a new texture */
 	bool ImportPage(int32 PageNum, const TCHAR* FileName);
 
-	/**	Spawns the viewport tab */
-	TSharedRef<SDockTab> SpawnTab_Viewport( const FSpawnTabArgs& Args );
+	/**	Spawns the text pages viewport tab */
+	TSharedRef<SDockTab> SpawnTab_TexturePagesViewport( const FSpawnTabArgs& Args );
+
+	/**	Spawns the composite font editor UI */
+	TSharedRef<SDockTab> SpawnTab_CompositeFontEditor( const FSpawnTabArgs& Args );
 
 	/**	Spawns the preview tab */
 	TSharedRef<SDockTab> SpawnTab_Preview( const FSpawnTabArgs& Args );
@@ -94,6 +106,9 @@ private:
 	/** Callback when an object is reimported, handles steps needed to keep the editor up-to-date. */
 	void OnObjectReimported(UObject* InObject);
 
+	/** Recreate the font object so that it's using the given caching method */
+	bool RecreateFontObject(const EFontCacheType NewCacheType);
+
 private:
 	/** The font asset being inspected */
 	UFont* Font;
@@ -104,6 +119,9 @@ private:
 	/** Viewport */
 	TSharedPtr<SFontEditorViewport> FontViewport;
 
+	/** Composite font editor UI */
+	TSharedPtr<SCompositeFontEditor> CompositeFontEditor;
+
 	/** Preview tab */
 	TSharedPtr<SVerticalBox> FontPreview;
 
@@ -128,8 +146,12 @@ private:
 	/** The factory to create updated pages with */
 	UTextureFactory* Factory;
 
+	/** The current font editor layout (if any) */
+	TOptional<EFontCacheType> CurrentEditorLayout;
+
 	/**	The tab ids for the font editor */
-	static const FName ViewportTabId;
+	static const FName TexturePagesViewportTabId;
+	static const FName CompositeFontEditorTabId;
 	static const FName PreviewTabId;
 	static const FName PropertiesTabId;
 	static const FName PagePropertiesTabId;
diff --git a/Engine/Source/Editor/FontEditor/Private/FontEditorModule.cpp b/Engine/Source/Editor/FontEditor/Private/FontEditorModule.cpp
--- a/Engine/Source/Editor/FontEditor/Private/FontEditorModule.cpp
+++ b/Engine/Source/Editor/FontEditor/Private/FontEditorModule.cpp
@@ -5,6 +5,7 @@
 #include "ModuleManager.h"
 #include "Factories.h"
 #include "SFontEditorViewport.h"
+#include "SCompositeFontEditor.h"
 #include "FontEditor.h"
 
 const FName FontEditorAppIdentifier = FName(TEXT("FontEditorApp"));
diff --git a/Engine/Source/Editor/FontEditor/Private/SCompositeFontEditor.cpp b/Engine/Source/Editor/FontEditor/Private/SCompositeFontEditor.cpp
new file mode 100644
--- /dev/null
+++ b/Engine/Source/Editor/FontEditor/Private/SCompositeFontEditor.cpp
@@ -0,0 +1,1284 @@
+// Copyright 1998-2014 Epic Games, Inc. All Rights Reserved.
+
+#include "FontEditorModule.h"
+#include "SCompositeFontEditor.h"
+#include "SFilePathPicker.h"
+#include "SInlineEditableTextBlock.h"
+#include "SNumericEntryBox.h"
+#include "ScopedTransaction.h"
+
+#define LOCTEXT_NAMESPACE "FontEditor"
+
+/** Entry used to weakly reference a particular typeface entry in the SListView */
+struct FTypefaceListViewEntry
+{
+	TAttribute<FTypeface*> Typeface;
+	int32 TypefaceEntryIndex;
+
+	FTypefaceListViewEntry()
+		: Typeface()
+		, TypefaceEntryIndex(INDEX_NONE)
+	{
+	}
+
+	FTypefaceListViewEntry(TAttribute<FTypeface*>& InTypeface, const int32 InTypefaceEntryIndex)
+		: Typeface(InTypeface)
+		, TypefaceEntryIndex(InTypefaceEntryIndex)
+	{
+	}
+
+	void Reset()
+	{
+		*this = FTypefaceListViewEntry();
+	}
+
+	FTypefaceEntry* GetTypefaceEntry() const
+	{
+		FTypeface* const TypefacePtr = Typeface.Get(nullptr);
+		return (TypefacePtr && TypefaceEntryIndex < TypefacePtr->Fonts.Num()) ? &TypefacePtr->Fonts[TypefaceEntryIndex] : nullptr;
+	}
+};
+
+/** Entry used to weakly reference a particular sub-typeface entry in the SListView */
+struct FSubTypefaceListViewEntry
+{
+	FCompositeFont* CompositeFont;
+	int32 SubTypefaceEntryIndex;
+
+	FSubTypefaceListViewEntry()
+		: CompositeFont(nullptr)
+		, SubTypefaceEntryIndex(INDEX_NONE)
+	{
+	}
+
+	FSubTypefaceListViewEntry(FCompositeFont* const InCompositeFont, const int32 InSubTypefaceEntryIndex)
+		: CompositeFont(InCompositeFont)
+		, SubTypefaceEntryIndex(InSubTypefaceEntryIndex)
+	{
+	}
+
+	void Reset()
+	{
+		*this = FSubTypefaceListViewEntry();
+	}
+
+	FCompositeSubFont* GetSubTypefaceEntry() const
+	{
+		return (CompositeFont && SubTypefaceEntryIndex < CompositeFont->SubTypefaces.Num()) ? &CompositeFont->SubTypefaces[SubTypefaceEntryIndex] : nullptr;
+	}
+};
+
+/** Entry used to weakly reference a particular character range entry in the STileView */
+struct FCharacterRangeTileViewEntry
+{
+	FSubTypefaceListViewEntryPtr SubTypefaceEntry;
+	int32 RangeEntryIndex;
+
+	FCharacterRangeTileViewEntry()
+		: SubTypefaceEntry()
+		, RangeEntryIndex(INDEX_NONE)
+	{
+	}
+
+	FCharacterRangeTileViewEntry(FSubTypefaceListViewEntryPtr InSubTypefaceEntry, const int32 InRangeEntryIndex)
+		: SubTypefaceEntry(InSubTypefaceEntry)
+		, RangeEntryIndex(InRangeEntryIndex)
+	{
+	}
+
+	void Reset()
+	{
+		*this = FCharacterRangeTileViewEntry();
+	}
+
+	FInt32Range* GetRange() const
+	{
+		FCompositeSubFont* const SubTypefaceEntryPtr = (SubTypefaceEntry.IsValid()) ? SubTypefaceEntry->GetSubTypefaceEntry() : nullptr;
+		return (SubTypefaceEntryPtr && RangeEntryIndex < SubTypefaceEntryPtr->CharacterRanges.Num()) ? &SubTypefaceEntryPtr->CharacterRanges[RangeEntryIndex] : nullptr;
+	}
+};
+
+SCompositeFontEditor::~SCompositeFontEditor()
+{
+}
+
+void SCompositeFontEditor::Construct(const FArguments& InArgs)
+{
+	FontEditorPtr = InArgs._FontEditor;
+
+	ChildSlot
+	[
+		SNew(SScrollBox)
+
+		+SScrollBox::Slot()
+		[
+			SNew(SVerticalBox)
+
+			+SVerticalBox::Slot()
+			.AutoHeight()
+			[
+				SAssignNew(DefaultTypefaceEditor, STypefaceEditor)
+				.CompositeFontEditor(this)
+				.Typeface(this, &SCompositeFontEditor::GetDefaultTypeface)
+				.TypefaceDisplayName(LOCTEXT("DefaultFontFamilyName", "Default Font Family"))
+			]
+
+			+SVerticalBox::Slot()
+			.AutoHeight()
+			[
+				SAssignNew(SubTypefaceEntriesListView, SListView<FSubTypefaceListViewEntryPtr>)
+				.ListItemsSource(&SubTypefaceEntries)
+				.SelectionMode(ESelectionMode::None)
+				.OnGenerateRow(this, &SCompositeFontEditor::MakeSubTypefaceEntryWidget)
+			]
+		]
+	];
+
+	UpdateSubTypefaceList();
+}
+
+void SCompositeFontEditor::Refresh()
+{
+	FlushCachedFont();
+
+	DefaultTypefaceEditor->Refresh();
+	UpdateSubTypefaceList();
+}
+
+void SCompositeFontEditor::FlushCachedFont()
+{
+	FCompositeFont* const CompositeFont = GetCompositeFont();
+	if(CompositeFont)
+	{
+		CompositeFont->MakeDirty();
+	}
+
+	TSharedPtr<IFontEditor> FontEditor = FontEditorPtr.Pin();
+	if(FontEditor.IsValid())
+	{
+		FontEditor->RefreshPreview();
+	}
+}
+
+UFont* SCompositeFontEditor::GetFontObject() const
+{
+	TSharedPtr<IFontEditor> FontEditor = FontEditorPtr.Pin();
+	return (FontEditor.IsValid()) ? FontEditor->GetFont() : nullptr;
+}
+
+FCompositeFont* SCompositeFontEditor::GetCompositeFont() const
+{
+	UFont* const FontObject = GetFontObject();
+	return (FontObject) ? &FontObject->CompositeFont : nullptr;
+}
+
+FTypeface* SCompositeFontEditor::GetDefaultTypeface() const
+{
+	UFont* const FontObject = GetFontObject();
+	return (FontObject) ? &FontObject->CompositeFont.DefaultTypeface : nullptr;
+}
+
+const FTypeface* SCompositeFontEditor::GetConstDefaultTypeface() const
+{
+	return GetDefaultTypeface();
+}
+
+void SCompositeFontEditor::UpdateSubTypefaceList()
+{
+	for(FSubTypefaceListViewEntryPtr& SubTypefaceListViewEntry : SubTypefaceEntries)
+	{
+		SubTypefaceListViewEntry->Reset();
+	}
+
+	FCompositeFont* const CompositeFontPtr = GetCompositeFont();
+	if(CompositeFontPtr)
+	{
+		SubTypefaceEntries.Empty(CompositeFontPtr->SubTypefaces.Num());
+
+		for(int32 SubTypefaceEntryIndex = 0; SubTypefaceEntryIndex < CompositeFontPtr->SubTypefaces.Num(); ++SubTypefaceEntryIndex)
+		{
+			SubTypefaceEntries.Add(MakeShareable(new FSubTypefaceListViewEntry(CompositeFontPtr, SubTypefaceEntryIndex)));
+		}
+	}
+
+	// Add a dummy entry for the "Add" button slot
+	SubTypefaceEntries.Add(MakeShareable(new FSubTypefaceListViewEntry()));
+
+	SubTypefaceEntriesListView->RequestListRefresh();
+}
+
+TSharedRef<ITableRow> SCompositeFontEditor::MakeSubTypefaceEntryWidget(FSubTypefaceListViewEntryPtr InSubTypefaceEntry, const TSharedRef<STableViewBase>& OwnerTable)
+{
+	TSharedPtr<SWidget> EntryWidget;
+
+	if(InSubTypefaceEntry->SubTypefaceEntryIndex == INDEX_NONE)
+	{
+		// Dummy entry for the "Add" button
+		EntryWidget = SNew(SBorder)
+		.BorderImage(FEditorStyle::GetBrush("ToolPanel.DarkGroupBorder"))
+		[
+			SNew(SButton)
+			.ButtonStyle(FEditorStyle::Get(), "HoverHintOnly")
+			.ForegroundColor(FSlateColor::UseForeground())
+			.ToolTipText(LOCTEXT("AddSubFontFamilyTooltip", "Add a sub-font family to this composite font"))
+			.OnClicked(this, &SCompositeFontEditor::OnAddSubFontFamily)
+			.VAlign(VAlign_Center)
+			[
+				SNew(SHorizontalBox)
+
+				+SHorizontalBox::Slot()
+				.AutoWidth()
+				.Padding(8.0f)
+				.VAlign(VAlign_Center)
+				[
+					SNew(SImage)
+					.Image(FEditorStyle::Get().GetBrush("FontEditor.Button_Add"))
+				]
+
+				+SHorizontalBox::Slot()
+				.AutoWidth()
+				.VAlign(VAlign_Center)
+				[
+					SNew(STextBlock)
+					.AutoWrapText(true)
+					.Text(LOCTEXT("AddSubFontFamily", "Add Sub-Font Family"))
+					.Font(FEditorStyle::GetFontStyle("DetailsView.CategoryFontStyle"))
+					.Justification(ETextJustify::Center)
+				]
+			]
+		];
+	}
+	else
+	{
+		const int32 SubFontIndex = SubTypefaceEntries.IndexOfByKey(InSubTypefaceEntry);
+		const FText SubFontFamilyName = FText::Format(LOCTEXT("SubFontFamilyNameFmt", "Sub-Font Family #{0}"), FText::AsNumber(SubFontIndex + 1));
+
+		EntryWidget = SNew(SSubTypefaceEditor)
+		.CompositeFontEditor(this)
+		.SubTypeface(InSubTypefaceEntry)
+		.ParentTypeface(this, &SCompositeFontEditor::GetConstDefaultTypeface)
+		.OnDeleteSubFontFamily(this, &SCompositeFontEditor::OnDeleteSubFontFamily)
+		.TypefaceDisplayName(SubFontFamilyName);
+	}
+
+	return
+		SNew(STableRow<FSubTypefaceListViewEntryPtr>, OwnerTable)
+		[
+			EntryWidget.ToSharedRef()
+		];
+}
+
+FReply SCompositeFontEditor::OnAddSubFontFamily()
+{
+	const FScopedTransaction Transaction(LOCTEXT("AddSubFontFamily", "Add Sub-Font Family"));
+	GetFontObject()->Modify();
+
+	FCompositeFont* const CompositeFontPtr = GetCompositeFont();
+	if(CompositeFontPtr)
+	{
+		CompositeFontPtr->SubTypefaces.Add(FCompositeSubFont());
+		UpdateSubTypefaceList();
+
+		FlushCachedFont();
+	}
+
+	return FReply::Handled();
+}
+
+void SCompositeFontEditor::OnDeleteSubFontFamily(const FSubTypefaceListViewEntryPtr& SubTypefaceEntryToRemove)
+{
+	const FScopedTransaction Transaction(LOCTEXT("DeleteSubFontFamily", "Delete Sub-Font Family"));
+	GetFontObject()->Modify();
+
+	FCompositeFont* const CompositeFontPtr = GetCompositeFont();
+	if(CompositeFontPtr)
+	{
+		CompositeFontPtr->SubTypefaces.RemoveAt(SubTypefaceEntryToRemove->SubTypefaceEntryIndex);
+		UpdateSubTypefaceList();
+
+		FlushCachedFont();
+	}
+}
+
+STypefaceEditor::~STypefaceEditor()
+{
+}
+
+void STypefaceEditor::Construct(const FArguments& InArgs)
+{
+	CompositeFontEditorPtr = InArgs._CompositeFontEditor;
+	Typeface = InArgs._Typeface;
+
+	ChildSlot
+	.Padding(FMargin(0.0f, 0.0f, 0.0f, 4.0f))
+	[
+		SNew(SBorder)
+		.BorderImage(FEditorStyle::GetBrush("ToolPanel.DarkGroupBorder"))
+		.Padding(0.0f)
+		[
+			SNew(SVerticalBox)
+
+			+SVerticalBox::Slot()
+			.AutoHeight()
+			.Padding(FMargin(8.0f, 8.0f, 16.0f, 8.0f))
+			[
+				SNew(SHorizontalBox)
+
+				+SHorizontalBox::Slot()
+				[
+					SNew(STextBlock)
+					.Text(InArgs._TypefaceDisplayName)
+					.Font(FEditorStyle::GetFontStyle("DetailsView.CategoryFontStyle"))
+				]
+
+				+SHorizontalBox::Slot()
+				.AutoWidth()
+				[
+					InArgs._HeaderContent.Widget
+				]
+			]
+
+			+SVerticalBox::Slot()
+			.AutoHeight()
+			[
+				InArgs._BodyContent.Widget
+			]
+
+			+SVerticalBox::Slot()
+			.AutoHeight()
+			[
+				SAssignNew(TypefaceEntriesListView, SListView<FTypefaceListViewEntryPtr>)
+				.ListItemsSource(&TypefaceEntries)
+				.SelectionMode(ESelectionMode::None)
+				.OnGenerateRow(this, &STypefaceEditor::MakeTypefaceEntryWidget)
+			]
+		]
+	];
+
+	UpdateFontList();
+}
+
+void STypefaceEditor::Refresh()
+{
+	UpdateFontList();
+}
+
+void STypefaceEditor::UpdateFontList()
+{
+	FTypeface* const TypefacePtr = Typeface.Get(nullptr);
+
+	for(FTypefaceListViewEntryPtr& TypefaceListViewEntry : TypefaceEntries)
+	{
+		TypefaceListViewEntry->Reset();
+	}
+
+	TypefaceEntries.Empty((TypefacePtr) ? TypefacePtr->Fonts.Num() : 0);
+
+	if(TypefacePtr)
+	{
+		for(int32 TypefaceEntryIndex = 0; TypefaceEntryIndex < TypefacePtr->Fonts.Num(); ++TypefaceEntryIndex)
+		{
+			TypefaceEntries.Add(MakeShareable(new FTypefaceListViewEntry(Typeface, TypefaceEntryIndex)));
+		}
+	}
+
+	// Add a dummy entry for the "Add" button slot
+	TypefaceEntries.Add(MakeShareable(new FTypefaceListViewEntry()));
+
+	TypefaceEntriesListView->RequestListRefresh();
+}
+
+TSharedRef<ITableRow> STypefaceEditor::MakeTypefaceEntryWidget(FTypefaceListViewEntryPtr InTypefaceEntry, const TSharedRef<STableViewBase>& OwnerTable)
+{
+	TSharedPtr<SWidget> EntryWidget;
+
+	if(InTypefaceEntry->TypefaceEntryIndex == INDEX_NONE)
+	{
+		// Dummy entry for the "Add" button
+		EntryWidget = SNew(SBorder)
+		.BorderImage(FEditorStyle::GetBrush("ToolPanel.GroupBorder"))
+		[
+			SNew(SButton)
+			.ButtonStyle(FEditorStyle::Get(), "HoverHintOnly")
+			.ForegroundColor(FSlateColor::UseForeground())
+			.ToolTipText(LOCTEXT("AddFontTooltip", "Add a new font to this font family"))
+			.OnClicked(this, &STypefaceEditor::OnAddFont)
+			[
+				SNew(SHorizontalBox)
+
+				+SHorizontalBox::Slot()
+				.AutoWidth()
+				.Padding(8.0f)
+				.VAlign(VAlign_Center)
+				[
+					SNew(SImage)
+					.Image(FEditorStyle::Get().GetBrush("FontEditor.Button_Add"))
+				]
+
+				+SHorizontalBox::Slot()
+				.AutoWidth()
+				.VAlign(VAlign_Center)
+				[
+					SNew(STextBlock)
+					.AutoWrapText(true)
+					.Text(LOCTEXT("AddFont", "Add Font"))
+					.Font(FEditorStyle::GetFontStyle("DetailsView.CategoryFontStyle"))
+					.Justification(ETextJustify::Center)
+				]
+			]
+		];
+	}
+	else
+	{
+		EntryWidget = SNew(STypefaceEntryEditor)
+		.CompositeFontEditor(CompositeFontEditorPtr)
+		.TypefaceEntry(InTypefaceEntry)
+		.OnDeleteFont(this, &STypefaceEditor::OnDeleteFont)
+		.OnVerifyFontName(this, &STypefaceEditor::OnVerifyFontName);
+	}
+
+	return
+		SNew(STableRow<FTypefaceListViewEntryPtr>, OwnerTable)
+		[
+			SNew(SBox)
+			.Padding(FMargin(8.0f, 0.0f, 8.0f, 8.0f))
+			[
+				EntryWidget.ToSharedRef()
+			]
+		];
+}
+
+FReply STypefaceEditor::OnAddFont()
+{
+	FTypeface* const TypefacePtr = Typeface.Get(nullptr);
+
+	if(TypefacePtr)
+	{
+		const FScopedTransaction Transaction(LOCTEXT("AddFont", "Add Font"));
+		CompositeFontEditorPtr->GetFontObject()->Modify();
+
+		TSet<FName> ExistingFontNames;
+		for(const FTypefaceEntry& TypefaceEntry : TypefacePtr->Fonts)
+		{
+			ExistingFontNames.Add(TypefaceEntry.Name);
+		}
+
+		// Get a valid default name for the font
+		static const FName BaseFontName("Font");
+		FName NewFontName = BaseFontName;
+		while(ExistingFontNames.Contains(NewFontName))
+		{
+			NewFontName.SetNumber(NewFontName.GetNumber() + 1);
+		}
+
+		TypefacePtr->Fonts.Add(FTypefaceEntry(NewFontName));
+		UpdateFontList();
+
+		CompositeFontEditorPtr->FlushCachedFont();
+	}
+
+	return FReply::Handled();
+}
+
+void STypefaceEditor::OnDeleteFont(const FTypefaceListViewEntryPtr& TypefaceEntryToRemove)
+{
+	FTypeface* const TypefacePtr = Typeface.Get(nullptr);
+
+	if(TypefacePtr && TypefaceEntryToRemove.IsValid() && TypefaceEntryToRemove->GetTypefaceEntry())
+	{
+		const FScopedTransaction Transaction(LOCTEXT("DeleteFont", "Delete Font"));
+		CompositeFontEditorPtr->GetFontObject()->Modify();
+
+		TypefacePtr->Fonts.RemoveAt(TypefaceEntryToRemove->TypefaceEntryIndex);
+		UpdateFontList();
+
+		CompositeFontEditorPtr->FlushCachedFont();
+	}
+}
+
+bool STypefaceEditor::OnVerifyFontName(const FTypefaceListViewEntryPtr& TypefaceEntryBeingRenamed, const FName& NewName, FText& OutFailureReason) const
+{
+	FTypeface* const TypefacePtr = Typeface.Get(nullptr);
+	FTypefaceEntry* const TypefaceEntryPtr = TypefaceEntryBeingRenamed->GetTypefaceEntry();
+
+	// Empty names are invalid
+	if(NewName.IsNone())
+	{
+		OutFailureReason = LOCTEXT("Error_FontNameEmpty", "The font name cannot be empty or 'None'");
+		return false;
+	}
+
+	// If we already have this name, it's valid
+	if(TypefaceEntryPtr && TypefaceEntryPtr->Name == NewName)
+	{
+		return true;
+	}
+
+	// Duplicate names are invalid
+	if(TypefacePtr)
+	{
+		const bool bNameExists = TypefacePtr->Fonts.ContainsByPredicate([&NewName](const FTypefaceEntry& TypefaceEntry) -> bool
+		{
+			return TypefaceEntry.Name == NewName;
+		});
+
+		if(bNameExists)
+		{
+			OutFailureReason = FText::Format(LOCTEXT("Error_DuplicateFontNameFmt", "A font with the name '{0}' already exists"), FText::FromName(NewName));
+			return false;
+		}
+	}
+
+	return true;
+}
+
+STypefaceEntryEditor::~STypefaceEntryEditor()
+{
+}
+
+void STypefaceEntryEditor::Construct(const FArguments& InArgs)
+{
+	CompositeFontEditorPtr = InArgs._CompositeFontEditor;
+	TypefaceEntry = InArgs._TypefaceEntry;
+	OnDeleteFont = InArgs._OnDeleteFont;
+	OnVerifyFontName = InArgs._OnVerifyFontName;
+
+	GatherHintingEnumEntries();
+
+	ChildSlot
+	[
+		SNew(SBorder)
+		.BorderImage(FEditorStyle::GetBrush("ToolPanel.GroupBorder"))
+		.Padding(8.0f)
+		[
+			SNew(SHorizontalBox)
+
+			+SHorizontalBox::Slot()
+			[
+				SNew(SVerticalBox)
+
+				+SVerticalBox::Slot()
+				.AutoHeight()
+				.Padding(FMargin(0.0f, 0.0f, 0.0f, 4.0f))
+				[
+					SNew(SHorizontalBox)
+
+					+SHorizontalBox::Slot()
+					[
+						SNew(SInlineEditableTextBlock)
+						.Text(this, &STypefaceEntryEditor::GetTypefaceEntryName)
+						.ToolTipText(LOCTEXT("FontNameTooltip", "The name of this font within the font family (click to edit)"))
+						.OnTextCommitted(this, &STypefaceEntryEditor::OnTypefaceEntryNameCommitted)
+						.OnVerifyTextChanged(this, &STypefaceEntryEditor::OnTypefaceEntryChanged)
+					]
+
+					+SHorizontalBox::Slot()
+					.AutoWidth()
+					.Padding(FMargin(4.0f, 0.0f, 0.0f, 0.0f))
+					[
+						SNew(SBorder)
+						.BorderImage(FEditorStyle::GetBrush("ToolPanel.DarkGroupBorder"))
+						.VAlign(VAlign_Center)
+						.HAlign(HAlign_Center)
+						.Padding(2.0f)
+						[
+							SAssignNew(PreviewTextBlock, STextBlock)
+							.Text(FText::FromString(TEXT("Preview")))
+							.ToolTipText(LOCTEXT("PreviewFontTooltip", "Preview of how this font will look when rendered by Slate"))
+							.Font(this, &STypefaceEntryEditor::GetPreviewFontStyle)
+						]
+					]
+				]
+
+				+SVerticalBox::Slot()
+				.AutoHeight()
+				[
+					SNew(SHorizontalBox)
+
+					+SHorizontalBox::Slot()
+					[
+						SNew(SFilePathPicker)
+						.BrowseButtonImage(FEditorStyle::GetBrush("PropertyWindow.Button_Ellipsis"))
+						.BrowseButtonStyle(FEditorStyle::Get(), "HoverHintOnly")
+						.BrowseButtonToolTip(LOCTEXT("FontFilePathPickerToolTip", "Choose a font file from this computer"))
+						.BrowseDirectory_Static([]() { return FEditorDirectories::Get().GetLastDirectory(ELastDirectory::GENERIC_OPEN); })
+						.BrowseTitle(LOCTEXT("FontPickerTitle", "Choose a font file..."))
+						.FilePath(this, &STypefaceEntryEditor::GetTypefaceEntryFontFilePath)
+						.FileTypeFilter(TEXT("TrueType fonts (*.ttf)|*.ttf|OpenType fonts (*.otf)|*.otf"))
+						.OnPathPicked(this, &STypefaceEntryEditor::OnTypefaceEntryFontPathPicked)
+					]
+
+					+SHorizontalBox::Slot()
+					.AutoWidth()
+					.Padding(FMargin(4.0f, 0.0f, 0.0f, 0.0f))
+					[
+						SNew(SBox)
+						.MinDesiredWidth(100.0f)
+						[
+							SNew(SComboBox<TSharedPtr<FFontHintingComboEntry>>)
+							.OptionsSource(&HintingComboData)
+							.OnSelectionChanged(this, &STypefaceEntryEditor::OnHintingComboSelectionChanged)
+							.OnGenerateWidget(this, &STypefaceEntryEditor::MakeHintingComboEntryWidget)
+							[
+								SNew(STextBlock)
+								.Text(this, &STypefaceEntryEditor::GetHintingComboText)
+							]
+						]
+					]
+				]
+			]
+
+			+SHorizontalBox::Slot()
+			.AutoWidth()
+			.VAlign(VAlign_Center)
+			.Padding(FMargin(8.0f, 0.0f, 0.0f, 0.0f))
+			[
+				SNew(SButton)
+				.ButtonStyle(FEditorStyle::Get(), "HoverHintOnly")
+				.ToolTipText(LOCTEXT("DeleteFontTooltip", "Remove this font from the font family"))
+				.OnClicked(this, &STypefaceEntryEditor::OnDeleteFontClicked)
+				[
+					SNew(SImage)
+					.Image(FEditorStyle::Get().GetBrush("FontEditor.Button_Delete"))
+				]
+			]
+		]
+	];
+}
+
+FText STypefaceEntryEditor::GetTypefaceEntryName() const
+{
+	FTypefaceEntry* const TypefaceEntryPtr = TypefaceEntry->GetTypefaceEntry();
+
+	if(TypefaceEntryPtr)
+	{
+		return FText::FromName(TypefaceEntryPtr->Name);
+	}
+
+	return FText::GetEmpty();
+}
+
+void STypefaceEntryEditor::OnTypefaceEntryNameCommitted(const FText& InNewName, ETextCommit::Type InCommitType)
+{
+	FTypefaceEntry* const TypefaceEntryPtr = TypefaceEntry->GetTypefaceEntry();
+
+	if(TypefaceEntryPtr)
+	{
+		const FScopedTransaction Transaction(LOCTEXT("RenameFont", "Rename Font"));
+		CompositeFontEditorPtr->GetFontObject()->Modify();
+
+		TypefaceEntryPtr->Name = *InNewName.ToString();
+
+		CompositeFontEditorPtr->FlushCachedFont();
+	}
+}
+
+bool STypefaceEntryEditor::OnTypefaceEntryChanged(const FText& InNewName, FText& OutFailureReason) const
+{
+	return !OnVerifyFontName.IsBound() || OnVerifyFontName.Execute(TypefaceEntry, *InNewName.ToString(), OutFailureReason);
+}
+
+FString STypefaceEntryEditor::GetTypefaceEntryFontFilePath() const
+{
+	FTypefaceEntry* const TypefaceEntryPtr = TypefaceEntry->GetTypefaceEntry();
+
+	if(TypefaceEntryPtr)
+	{
+		return TypefaceEntryPtr->Font.FontFilename;
+	}
+
+	return FString();
+}
+
+void STypefaceEntryEditor::OnTypefaceEntryFontPathPicked(const FString& InNewFontFilename)
+{
+	FTypefaceEntry* const TypefaceEntryPtr = TypefaceEntry->GetTypefaceEntry();
+
+	if(TypefaceEntryPtr)
+	{
+		const FScopedTransaction Transaction(LOCTEXT("SetFontFile", "Set Font File"));
+		CompositeFontEditorPtr->GetFontObject()->Modify();
+
+		if(TypefaceEntryPtr->Font.SetFont(InNewFontFilename))
+		{
+			CompositeFontEditorPtr->FlushCachedFont();
+		}
+	}
+
+	FEditorDirectories::Get().SetLastDirectory(ELastDirectory::GENERIC_OPEN, FPaths::GetPath(InNewFontFilename));
+}
+
+FReply STypefaceEntryEditor::OnDeleteFontClicked()
+{
+	OnDeleteFont.ExecuteIfBound(TypefaceEntry);
+
+	return FReply::Handled();
+}
+
+FSlateFontInfo STypefaceEntryEditor::GetPreviewFontStyle() const
+{
+	FTypefaceEntry* const TypefaceEntryPtr = TypefaceEntry->GetTypefaceEntry();
+	return FSlateFontInfo(CompositeFontEditorPtr->GetFontObject(), 9, (TypefaceEntryPtr) ? TypefaceEntryPtr->Name : NAME_None);
+}
+
+void STypefaceEntryEditor::GatherHintingEnumEntries()
+{
+	FTypefaceEntry* const TypefaceEntryPtr = TypefaceEntry->GetTypefaceEntry();
+
+	const UEnum* const HintingEnum = FindObject<UEnum>(ANY_PACKAGE, TEXT("EFontHinting"), true);
+	check(HintingEnum);
+
+	for(int32 EnumIndex = 0; EnumIndex < HintingEnum->NumEnums() - 1; ++EnumIndex)
+	{
+		// Ignore hidden enum entries
+		const bool bShouldBeHidden = HintingEnum->HasMetaData(TEXT("Hidden"), EnumIndex) || HintingEnum->HasMetaData(TEXT("Spacer"), EnumIndex);
+		if(bShouldBeHidden)
+		{
+			continue;
+		}
+
+		TSharedPtr<FFontHintingComboEntry> ComboEntry = MakeShareable(new FFontHintingComboEntry());
+		ComboEntry->EnumValue = static_cast<EFontHinting>(EnumIndex);
+
+		// See if we specified an alternate name for this entry using metadata
+		ComboEntry->DisplayName = HintingEnum->GetDisplayNameText(EnumIndex);
+		if(ComboEntry->DisplayName.IsEmpty()) 
+		{
+			ComboEntry->DisplayName = HintingEnum->GetEnumText(EnumIndex);
+		}
+
+		ComboEntry->Tooltip = HintingEnum->GetToolTipText(EnumIndex);
+
+		if(TypefaceEntryPtr && ComboEntry->EnumValue == TypefaceEntryPtr->Font.Hinting)
+		{
+			ActiveHintingEnumEntryText = ComboEntry->DisplayName;
+		}
+
+		HintingComboData.Add(ComboEntry);
+	}
+}
+
+void STypefaceEntryEditor::OnHintingComboSelectionChanged(TSharedPtr<FFontHintingComboEntry> InNewSelection, ESelectInfo::Type)
+{
+	FTypefaceEntry* const TypefaceEntryPtr = TypefaceEntry->GetTypefaceEntry();
+
+	if(TypefaceEntryPtr && InNewSelection.IsValid() && TypefaceEntryPtr->Font.Hinting != InNewSelection->EnumValue)
+	{
+		const FScopedTransaction Transaction(LOCTEXT("SetFontHinting", "Set Font Hinting"));
+		CompositeFontEditorPtr->GetFontObject()->Modify();
+
+		TypefaceEntryPtr->Font.Hinting = InNewSelection->EnumValue;
+		ActiveHintingEnumEntryText = InNewSelection->DisplayName;
+
+		CompositeFontEditorPtr->FlushCachedFont();
+	}
+}
+
+TSharedRef<SWidget> STypefaceEntryEditor::MakeHintingComboEntryWidget(TSharedPtr<FFontHintingComboEntry> InHintingEntry)
+{
+	return
+		SNew(STextBlock)
+		.Text(InHintingEntry->DisplayName)
+		.ToolTipText(InHintingEntry->Tooltip);
+}
+
+FText STypefaceEntryEditor::GetHintingComboText() const
+{
+	return ActiveHintingEnumEntryText;
+}
+
+SSubTypefaceEditor::~SSubTypefaceEditor()
+{
+}
+
+void SSubTypefaceEditor::Construct(const FArguments& InArgs)
+{
+	CompositeFontEditorPtr = InArgs._CompositeFontEditor;
+	SubTypeface = InArgs._SubTypeface;
+	ParentTypeface = InArgs._ParentTypeface;
+	OnDeleteSubFontFamily = InArgs._OnDeleteSubFontFamily;
+
+	ChildSlot
+	[
+		SAssignNew(TypefaceEditor, STypefaceEditor)
+		.CompositeFontEditor(InArgs._CompositeFontEditor)
+		.Typeface(this, &SSubTypefaceEditor::GetTypeface)
+		.TypefaceDisplayName(InArgs._TypefaceDisplayName)
+		.HeaderContent()
+		[
+			SNew(SBox)
+			.VAlign(VAlign_Center)
+			[
+				SNew(SHorizontalBox)
+
+				+SHorizontalBox::Slot()
+				.AutoWidth()
+				.VAlign(VAlign_Center)
+				.Padding(FMargin(4.0f, 0.0f))
+				[
+					SNew(SNumericEntryBox<float>)
+					.ToolTipText(LOCTEXT("ScalingFactorTooltip", "The scaling factor will adjust the size of the rendered glyphs so that you can tweak their size to match that of the default font family"))
+					.Value(this, &SSubTypefaceEditor::GetScalingFactorAsOptional)
+					.OnValueCommitted(this, &SSubTypefaceEditor::OnScalingFactorCommittedAsNumeric)
+					.LabelVAlign(VAlign_Center)
+					.Label()
+					[
+						SNew(STextBlock)
+						.Text(LOCTEXT("ScalingFactorLabel", "Scaling Factor"))
+					]
+				]
+
+				+SHorizontalBox::Slot()
+				.AutoWidth()
+				.VAlign(VAlign_Center)
+				[
+					SAssignNew(FontOverrideCombo, SComboBox<TSharedPtr<FName>>)
+					.OptionsSource(&FontOverrideComboData)
+					.ContentPadding(FMargin(4.0, 2.0))
+					.Visibility(this, &SSubTypefaceEditor::GetAddFontOverrideVisibility)
+					.OnComboBoxOpening(this, &SSubTypefaceEditor::OnAddFontOverrideComboOpening)
+					.OnSelectionChanged(this, &SSubTypefaceEditor::OnAddFontOverrideSelectionChanged)
+					.OnGenerateWidget(this, &SSubTypefaceEditor::MakeAddFontOverrideWidget)
+					[
+						SNew(STextBlock)
+						.Text(LOCTEXT("AddFontOverride", "Add Font Override"))
+						.ToolTipText(LOCTEXT("AddFontOverrideTooltip", "Override a font from the default font family to ensure it will be used when drawing a glyph in the range of this sub-font family"))
+					]
+				]
+
+				+SHorizontalBox::Slot()
+				.AutoWidth()
+				.VAlign(VAlign_Center)
+				.Padding(FMargin(8.0f, 0.0f, 0.0f, 0.0f))
+				[
+					SNew(SButton)
+					.ButtonStyle(FEditorStyle::Get(), "HoverHintOnly")
+					.ToolTipText(LOCTEXT("DeleteFontTooltip", "Remove this sub-font family from the composite font"))
+					.OnClicked(this, &SSubTypefaceEditor::OnDeleteSubFontFamilyClicked)
+					[
+						SNew(SImage)
+						.Image(FEditorStyle::Get().GetBrush("FontEditor.Button_Delete"))
+					]
+				]
+			]
+		]
+		.BodyContent()
+		[
+			SNew(SBox)
+			.Padding(FMargin(8.0f, 0.0f, 8.0f, 0.0f))
+			[
+				SAssignNew(CharacterRangeEntriesTileView, STileView<FCharacterRangeTileViewEntryPtr>)
+				.ListItemsSource(&CharacterRangeEntries)
+				.SelectionMode(ESelectionMode::None)
+				.ItemWidth(160)
+				.ItemHeight(120)
+				.ItemAlignment(EListItemAlignment::LeftAligned)
+				.OnGenerateTile(this, &SSubTypefaceEditor::MakeCharacterRangesEntryWidget)
+			]
+		]
+	];
+
+	UpdateCharacterRangesList();
+}
+
+FTypeface* SSubTypefaceEditor::GetTypeface() const
+{
+	FCompositeSubFont* const SubTypefaceEntryPtr = SubTypeface->GetSubTypefaceEntry();
+	return (SubTypefaceEntryPtr) ? &SubTypefaceEntryPtr->Typeface : nullptr;
+}
+
+EVisibility SSubTypefaceEditor::GetAddFontOverrideVisibility() const
+{
+	return (ParentTypeface.Get(nullptr)) ? EVisibility::Visible : EVisibility::Collapsed;
+}
+
+void SSubTypefaceEditor::OnAddFontOverrideComboOpening()
+{
+	FontOverrideComboData.Empty();
+	
+	FCompositeSubFont* const SubTypefaceEntryPtr = SubTypeface->GetSubTypefaceEntry();
+	const FTypeface* const ParentTypefacePtr = ParentTypeface.Get(nullptr);
+
+	if(SubTypefaceEntryPtr && ParentTypefacePtr)
+	{
+		TSet<FName> LocalFontNames;
+		for(const FTypefaceEntry& LocalTypefaceEntry : SubTypefaceEntryPtr->Typeface.Fonts)
+		{
+			LocalFontNames.Add(LocalTypefaceEntry.Name);
+		}
+
+		// Add every font from our parent font that hasn't already got a local entry
+		for(const FTypefaceEntry& ParentTypefaceEntry : ParentTypefacePtr->Fonts)
+		{
+			if(!LocalFontNames.Contains(ParentTypefaceEntry.Name))
+			{
+				FontOverrideComboData.Add(MakeShareable(new FName(ParentTypefaceEntry.Name)));
+			}
+		}
+	}
+
+	FontOverrideCombo->RefreshOptions();
+}
+
+void SSubTypefaceEditor::OnAddFontOverrideSelectionChanged(TSharedPtr<FName> InNewSelection, ESelectInfo::Type)
+{
+	FCompositeSubFont* const SubTypefaceEntryPtr = SubTypeface->GetSubTypefaceEntry();
+
+	if(SubTypefaceEntryPtr && InNewSelection.IsValid() && !InNewSelection->IsNone())
+	{
+		const FScopedTransaction Transaction(LOCTEXT("AddFontOverride", "Add Font Override"));
+		CompositeFontEditorPtr->GetFontObject()->Modify();
+
+		SubTypefaceEntryPtr->Typeface.Fonts.Add(FTypefaceEntry(*InNewSelection));
+		TypefaceEditor->Refresh();
+
+		CompositeFontEditorPtr->FlushCachedFont();
+	}
+}
+
+TSharedRef<SWidget> SSubTypefaceEditor::MakeAddFontOverrideWidget(TSharedPtr<FName> InFontEntry)
+{
+	return
+		SNew(STextBlock)
+		.Text(FText::FromName(*InFontEntry));
+}
+
+FReply SSubTypefaceEditor::OnDeleteSubFontFamilyClicked()
+{
+	OnDeleteSubFontFamily.ExecuteIfBound(SubTypeface);
+
+	return FReply::Handled();
+}
+
+void SSubTypefaceEditor::UpdateCharacterRangesList()
+{
+	FCompositeSubFont* const SubTypefaceEntryPtr = SubTypeface->GetSubTypefaceEntry();
+
+	for(FCharacterRangeTileViewEntryPtr& CharacterRangeTileViewEntry : CharacterRangeEntries)
+	{
+		CharacterRangeTileViewEntry->Reset();
+	}
+
+	CharacterRangeEntries.Empty((SubTypefaceEntryPtr) ? SubTypefaceEntryPtr->CharacterRanges.Num() : 0);
+
+	if(SubTypefaceEntryPtr)
+	{
+		for(int32 CharacterRangeIndex = 0; CharacterRangeIndex < SubTypefaceEntryPtr->CharacterRanges.Num(); ++CharacterRangeIndex)
+		{
+			CharacterRangeEntries.Add(MakeShareable(new FCharacterRangeTileViewEntry(SubTypeface, CharacterRangeIndex)));
+		}
+	}
+
+	// Add a dummy entry for the "Add" button slot
+	CharacterRangeEntries.Add(MakeShareable(new FCharacterRangeTileViewEntry()));
+
+	CharacterRangeEntriesTileView->RequestListRefresh();
+}
+
+TSharedRef<ITableRow> SSubTypefaceEditor::MakeCharacterRangesEntryWidget(FCharacterRangeTileViewEntryPtr InCharacterRangeEntry, const TSharedRef<STableViewBase>& OwnerTable)
+{
+	TSharedPtr<SWidget> EntryWidget;
+
+	if(InCharacterRangeEntry->RangeEntryIndex == INDEX_NONE)
+	{
+		// Dummy entry for the "Add" button
+		EntryWidget = SNew(SBorder)
+		.BorderImage(FEditorStyle::GetBrush("ToolPanel.GroupBorder"))
+		[
+			SNew(SButton)
+			.ButtonStyle(FEditorStyle::Get(), "HoverHintOnly")
+			.ForegroundColor(FSlateColor::UseForeground())
+			.ToolTipText(LOCTEXT("AddCharacterRangeTooltip", "Add a new character range to this sub-font family"))
+			.OnClicked(this, &SSubTypefaceEditor::OnAddCharacterRangeClicked)
+			.VAlign(VAlign_Center)
+			[
+				SNew(SVerticalBox)
+
+				+SVerticalBox::Slot()
+				.AutoHeight()
+				.Padding(16.0f)
+				.HAlign(HAlign_Center)
+				[
+					SNew(SImage)
+					.Image(FEditorStyle::Get().GetBrush("FontEditor.Button_Add"))
+				]
+
+				+SVerticalBox::Slot()
+				.AutoHeight()
+				.HAlign(HAlign_Center)
+				[
+					SNew(STextBlock)
+					.AutoWrapText(true)
+					.Text(LOCTEXT("AddCharacterRange", "Add Character Range"))
+					.Font(FEditorStyle::GetFontStyle("DetailsView.CategoryFontStyle"))
+					.Justification(ETextJustify::Center)
+				]
+			]
+		];
+	}
+	else
+	{
+		EntryWidget = SNew(SBorder)
+		.BorderImage(FEditorStyle::GetBrush("ToolPanel.GroupBorder"))
+		.Padding(8.0f)
+		[
+			SNew(SVerticalBox)
+
+			+SVerticalBox::Slot()
+			.VAlign(VAlign_Center)
+			.HAlign(HAlign_Center)
+			[
+				SNew(SCharacterRangeEditor)
+				.CompositeFontEditor(CompositeFontEditorPtr)
+				.CharacterRange(InCharacterRangeEntry)
+			]
+
+			+SVerticalBox::Slot()
+			.AutoHeight()
+			.HAlign(HAlign_Center)
+			[
+				SNew(SButton)
+				.ButtonStyle(FEditorStyle::Get(), "HoverHintOnly")
+				.ToolTipText(LOCTEXT("DeleteCharacterRangeTooltip", "Remove this character range from the sub-font family"))
+				.OnClicked(this, &SSubTypefaceEditor::OnDeleteCharacterRangeClicked, InCharacterRangeEntry)
+				[
+					SNew(SImage)
+					.Image(FEditorStyle::Get().GetBrush("FontEditor.Button_Delete"))
+				]
+			]
+		];
+	}
+
+	return
+		SNew(STableRow<FCharacterRangeTileViewEntryPtr>, OwnerTable)
+		[
+			SNew(SBox)
+			.Padding(FMargin(0.0f, 0.0f, 8.0f, 8.0f))
+			[
+				EntryWidget.ToSharedRef()
+			]
+		];
+}
+
+FReply SSubTypefaceEditor::OnAddCharacterRangeClicked()
+{
+	FCompositeSubFont* const SubTypefaceEntryPtr = SubTypeface->GetSubTypefaceEntry();
+
+	if(SubTypefaceEntryPtr)
+	{
+		const FScopedTransaction Transaction(LOCTEXT("AddCharacterRange", "Add Character Range"));
+		CompositeFontEditorPtr->GetFontObject()->Modify();
+
+		SubTypefaceEntryPtr->CharacterRanges.Add(FInt32Range::Empty());
+
+		UpdateCharacterRangesList();
+
+		CompositeFontEditorPtr->FlushCachedFont();
+	}
+
+	return FReply::Handled();
+}
+
+FReply SSubTypefaceEditor::OnDeleteCharacterRangeClicked(FCharacterRangeTileViewEntryPtr InCharacterRangeEntry)
+{
+	FCompositeSubFont* const SubTypefaceEntryPtr = SubTypeface->GetSubTypefaceEntry();
+
+	if(SubTypefaceEntryPtr)
+	{
+		const FScopedTransaction Transaction(LOCTEXT("DeleteCharacterRange", "Delete Character Range"));
+		CompositeFontEditorPtr->GetFontObject()->Modify();
+
+		SubTypefaceEntryPtr->CharacterRanges.RemoveAt(InCharacterRangeEntry->RangeEntryIndex);
+
+		UpdateCharacterRangesList();
+
+		CompositeFontEditorPtr->FlushCachedFont();
+	}
+
+	return FReply::Handled();
+}
+
+TOptional<float> SSubTypefaceEditor::GetScalingFactorAsOptional() const
+{
+	const FCompositeSubFont* const SubTypefaceEntryPtr = SubTypeface->GetSubTypefaceEntry();
+
+	if(SubTypefaceEntryPtr)
+	{
+		return SubTypefaceEntryPtr->ScalingFactor;
+	}
+
+	return TOptional<float>();
+}
+
+void SSubTypefaceEditor::OnScalingFactorCommittedAsNumeric(float InNewValue, ETextCommit::Type InCommitType)
+{
+	FCompositeSubFont* const SubTypefaceEntryPtr = SubTypeface->GetSubTypefaceEntry();
+
+	if(SubTypefaceEntryPtr)
+	{
+		const FScopedTransaction Transaction(LOCTEXT("SetScalingFactor", "Set Scaling Factor"));
+		CompositeFontEditorPtr->GetFontObject()->Modify();
+
+		SubTypefaceEntryPtr->ScalingFactor = InNewValue;
+
+		CompositeFontEditorPtr->FlushCachedFont();
+	}
+}
+
+SCharacterRangeEditor::~SCharacterRangeEditor()
+{
+}
+
+void SCharacterRangeEditor::Construct(const FArguments& InArgs)
+{
+	CompositeFontEditorPtr = InArgs._CompositeFontEditor;
+	CharacterRange = InArgs._CharacterRange;
+
+	ChildSlot
+	[
+		SNew(SGridPanel)
+
+		// Minimum column
+		+SGridPanel::Slot(0, 0)
+		.Padding(2.0f)
+		[
+			SNew(SEditableTextBox)
+			.Text(this, &SCharacterRangeEditor::GetRangeComponentAsTCHAR, 0)
+			.OnTextCommitted(this, &SCharacterRangeEditor::OnRangeComponentCommittedAsTCHAR, 0)
+		]
+
+		+SGridPanel::Slot(0, 1)
+		.Padding(2.0f)
+		[
+			SNew(SEditableTextBox)
+			.Text(this, &SCharacterRangeEditor::GetRangeComponentAsHexString, 0)
+			.OnTextCommitted(this, &SCharacterRangeEditor::OnRangeComponentCommittedAsHexString, 0)
+		]
+
+		+SGridPanel::Slot(0, 2)
+		.Padding(2.0f)
+		[
+			SNew(SNumericEntryBox<int32>)
+			.Value(this, &SCharacterRangeEditor::GetRangeComponentAsOptional, 0)
+			.OnValueCommitted(this, &SCharacterRangeEditor::OnRangeComponentCommittedAsNumeric, 0)
+		]
+
+		// Separator
+		+SGridPanel::Slot(1, 0)
+		.RowSpan(3)
+		.VAlign(VAlign_Center)
+		[
+			SNew(STextBlock)
+			.Text(FText::FromString(TEXT(" - ")))
+			.Font(FEditorStyle::GetFontStyle("DetailsView.CategoryFontStyle"))
+		]
+
+		// Maximum column
+		+SGridPanel::Slot(2, 0)
+		.Padding(2.0f)
+		[
+			SNew(SEditableTextBox)
+			.Text(this, &SCharacterRangeEditor::GetRangeComponentAsTCHAR, 1)
+			.OnTextCommitted(this, &SCharacterRangeEditor::OnRangeComponentCommittedAsTCHAR, 1)
+		]
+
+		+SGridPanel::Slot(2, 1)
+		.Padding(2.0f)
+		[
+			SNew(SEditableTextBox)
+			.Text(this, &SCharacterRangeEditor::GetRangeComponentAsHexString, 1)
+			.OnTextCommitted(this, &SCharacterRangeEditor::OnRangeComponentCommittedAsHexString, 1)
+		]
+
+		+SGridPanel::Slot(2, 2)
+		.Padding(2.0f)
+		[
+			SNew(SNumericEntryBox<int32>)
+			.Value(this, &SCharacterRangeEditor::GetRangeComponentAsOptional, 1)
+			.OnValueCommitted(this, &SCharacterRangeEditor::OnRangeComponentCommittedAsNumeric, 1)
+		]
+	];
+}
+
+FText SCharacterRangeEditor::GetRangeComponentAsTCHAR(int32 InComponentIndex) const
+{
+	const int32 RangeComponent = GetRangeComponent(InComponentIndex);
+	const TCHAR RangeComponentStr[] = { static_cast<TCHAR>(RangeComponent), 0 };
+	return FText::FromString(RangeComponentStr);
+}
+
+FText SCharacterRangeEditor::GetRangeComponentAsHexString(int32 InComponentIndex) const
+{
+	const int32 RangeComponent = GetRangeComponent(InComponentIndex);
+	return FText::FromString(FString::Printf(TEXT("0x%04x"), RangeComponent));
+}
+
+TOptional<int32> SCharacterRangeEditor::GetRangeComponentAsOptional(int32 InComponentIndex) const
+{
+	return GetRangeComponent(InComponentIndex);
+}
+
+int32 SCharacterRangeEditor::GetRangeComponent(const int32 InComponentIndex) const
+{
+	check(InComponentIndex == 0 || InComponentIndex == 1);
+
+	const FInt32Range* const CharacterRangePtr = CharacterRange->GetRange();
+	if(CharacterRangePtr)
+	{
+		return (InComponentIndex == 0) ? CharacterRangePtr->GetLowerBoundValue() : CharacterRangePtr->GetUpperBoundValue();
+	}
+
+	return 0;
+}
+
+void SCharacterRangeEditor::OnRangeComponentCommittedAsTCHAR(const FText& InNewValue, ETextCommit::Type InCommitType, int32 InComponentIndex)
+{
+	const FString NewValueStr = InNewValue.ToString();
+	if(NewValueStr.Len() == 1)
+	{
+		SetRangeComponent(static_cast<int32>(NewValueStr[0]), InComponentIndex);
+	}
+	else if(NewValueStr.Len() == 0)
+	{
+		SetRangeComponent(0, InComponentIndex);
+	}
+}
+
+void SCharacterRangeEditor::OnRangeComponentCommittedAsHexString(const FText& InNewValue, ETextCommit::Type InCommitType, int32 InComponentIndex)
+{
+	const FString NewValueStr = InNewValue.ToString();
+	const TCHAR* HexStart = NewValueStr.GetCharArray().GetData();
+	if(NewValueStr.StartsWith(TEXT("0x")))
+	{
+		// Skip the "0x" part, as FParse::HexNumber doesn't handle that
+		HexStart += 2;
+	}
+
+	const int32 NewValue = FParse::HexNumber(HexStart);
+	SetRangeComponent(NewValue, InComponentIndex);
+}
+
+void SCharacterRangeEditor::OnRangeComponentCommittedAsNumeric(int32 InNewValue, ETextCommit::Type InCommitType, int32 InComponentIndex)
+{
+	SetRangeComponent(InNewValue, InComponentIndex);
+}
+
+void SCharacterRangeEditor::SetRangeComponent(const int32 InNewValue, const int32 InComponentIndex)
+{
+	check(InComponentIndex == 0 || InComponentIndex == 1);
+	
+	FInt32Range* const CharacterRangePtr = CharacterRange->GetRange();
+	if(CharacterRangePtr)
+	{
+		const FScopedTransaction Transaction(LOCTEXT("UpdateCharacterRange", "Update Character Range"));
+		CompositeFontEditorPtr->GetFontObject()->Modify();
+
+		*CharacterRangePtr = (InComponentIndex == 0) 
+			? FInt32Range(FInt32Range::BoundsType::Inclusive(InNewValue), FInt32Range::BoundsType::Inclusive(CharacterRangePtr->GetUpperBoundValue())) 
+			: FInt32Range(FInt32Range::BoundsType::Inclusive(CharacterRangePtr->GetLowerBoundValue()), FInt32Range::BoundsType::Inclusive(InNewValue));
+
+		CompositeFontEditorPtr->FlushCachedFont();
+	}
+}
+
+#undef LOCTEXT_NAMESPACE
diff --git a/Engine/Source/Editor/FontEditor/Private/SCompositeFontEditor.h b/Engine/Source/Editor/FontEditor/Private/SCompositeFontEditor.h
new file mode 100644
--- /dev/null
+++ b/Engine/Source/Editor/FontEditor/Private/SCompositeFontEditor.h
@@ -0,0 +1,396 @@
+// Copyright 1998-2014 Epic Games, Inc. All Rights Reserved.
+
+#pragma once
+
+struct FTypefaceListViewEntry;
+typedef TSharedPtr<FTypefaceListViewEntry> FTypefaceListViewEntryPtr;
+
+struct FSubTypefaceListViewEntry;
+typedef TSharedPtr<FSubTypefaceListViewEntry> FSubTypefaceListViewEntryPtr;
+
+struct FCharacterRangeTileViewEntry;
+typedef TSharedPtr<FCharacterRangeTileViewEntry> FCharacterRangeTileViewEntryPtr;
+
+class SCompositeFontEditor : public SCompoundWidget
+{
+public:
+	SLATE_BEGIN_ARGS(SCompositeFontEditor)
+		: _FontEditor()
+		{}
+
+		/** Pointer back to the Font editor tool that owns us */
+		SLATE_ARGUMENT(TWeakPtr<IFontEditor>, FontEditor)
+
+	SLATE_END_ARGS()
+
+	/** Destructor */
+	~SCompositeFontEditor();
+
+	/** SCompoundWidget interface */
+	void Construct(const FArguments& InArgs);
+
+	/** Refresh the editor in response to an external change */
+	void Refresh();
+
+	/** Flush the font cache for the current font */
+	void FlushCachedFont();
+
+	/** Get the root font object being edited */
+	UFont* GetFontObject() const;
+
+private:
+	/** Get the composite font being edited */
+	FCompositeFont* GetCompositeFont() const;
+
+	/** Get the default typeface */
+	FTypeface* GetDefaultTypeface() const;
+
+	/** Get the default typeface */
+	const FTypeface* GetConstDefaultTypeface() const;
+
+	/** Update the list of sub-typefaces in this composite font */
+	void UpdateSubTypefaceList();
+
+	/** Make the widget for an entry in the sub-typeface entries list view */
+	TSharedRef<ITableRow> MakeSubTypefaceEntryWidget(FSubTypefaceListViewEntryPtr InSubTypefaceEntry, const TSharedRef<STableViewBase>& OwnerTable);
+
+	/** Called in response to the "Add Sub-Font Family" button being clicked" */
+	FReply OnAddSubFontFamily();
+
+	/** Delete the given sub-font family from this composite font */
+	void OnDeleteSubFontFamily(const FSubTypefaceListViewEntryPtr& SubTypefaceEntryToRemove);
+
+	/** Pointer back to the Font editor tool that owns us */
+	TWeakPtr<IFontEditor> FontEditorPtr;
+
+	/** Widget for editing the default typeface */
+	TSharedPtr<class STypefaceEditor> DefaultTypefaceEditor;
+
+	/** Internal list of sub-typeface pointers for the list view (generated from CompositeFontPtr->SubTypefaces) */
+	TArray<FSubTypefaceListViewEntryPtr> SubTypefaceEntries;
+
+	/** List view widget showing the sub-typeface editors (uses SubTypefaceEntries as its source) */
+	TSharedPtr<SListView<FSubTypefaceListViewEntryPtr>> SubTypefaceEntriesListView;
+};
+
+class STypefaceEditor : public SCompoundWidget
+{
+public:
+	SLATE_BEGIN_ARGS(STypefaceEditor)
+		: _CompositeFontEditor(nullptr)
+		, _Typeface(nullptr)
+		, _TypefaceDisplayName()
+		, _HeaderContent()
+		, _BodyContent()
+		{}
+
+		/** Pointer back to the composite font editor that owns us */
+		SLATE_ARGUMENT(SCompositeFontEditor*, CompositeFontEditor)
+
+		/** Typeface to edit (may be invalid, or change in response to an undo/redo) */
+		SLATE_ATTRIBUTE(FTypeface*, Typeface)
+
+		/** Name to show in the UI for this typeface */
+		SLATE_ATTRIBUTE(FText, TypefaceDisplayName)
+
+		/** Slot for extra content to place in the header bar (optional) */
+		SLATE_NAMED_SLOT(FArguments, HeaderContent)
+
+		/** Slot for extra content to place above the typeface entry list (optional) */
+		SLATE_NAMED_SLOT(FArguments, BodyContent)
+
+	SLATE_END_ARGS()
+
+	/** Destructor */
+	~STypefaceEditor();
+
+	/** SCompoundWidget interface */
+	void Construct(const FArguments& InArgs);
+
+	/** Refresh the editor in response to an external change */
+	void Refresh();
+
+private:
+	/** Update the list of fonts in this typeface */
+	void UpdateFontList();
+
+	/** Make the widget for an entry in the typeface entries list view */
+	TSharedRef<ITableRow> MakeTypefaceEntryWidget(FTypefaceListViewEntryPtr InTypefaceEntry, const TSharedRef<STableViewBase>& OwnerTable);
+
+	/** Called in response to the "Add Font" button being clicked" */
+	FReply OnAddFont();
+
+	/** Delete the given entry from this typeface */
+	void OnDeleteFont(const FTypefaceListViewEntryPtr& TypefaceEntryToRemove);
+
+	/** Verify that the name of the given typeface entry is valid */
+	bool OnVerifyFontName(const FTypefaceListViewEntryPtr& TypefaceEntryBeingRenamed, const FName& NewName, FText& OutFailureReason) const;
+
+	/** Pointer back to the composite font editor that owns us */
+	SCompositeFontEditor* CompositeFontEditorPtr;
+
+	/** Typeface to edit (may be invalid, or change in response to an undo/redo) */
+	TAttribute<FTypeface*> Typeface;
+
+	/** Internal list of font pointers for the list view (generated from TypefacePtr->Fonts) */
+	TArray<FTypefaceListViewEntryPtr> TypefaceEntries;
+
+	/** List view widget showing the font entries (uses TypefaceEntries as its source) */
+	TSharedPtr<SListView<FTypefaceListViewEntryPtr>> TypefaceEntriesListView;
+};
+
+class STypefaceEntryEditor : public SCompoundWidget
+{
+public:
+	/** Declares a delegate that is executed when this typeface entry should be deleted */
+	DECLARE_DELEGATE_OneParam(FOnDeleteFont, const FTypefaceListViewEntryPtr& /*TypefaceEntryToRemove*/);
+
+	/** Declares a delegate that is executed when this typeface entry name is changed - used to verify that the new name is valid */
+	DECLARE_DELEGATE_RetVal_ThreeParams(bool, FOnVerifyFontName, const FTypefaceListViewEntryPtr& /*TypefaceEntryBeingRenamed*/, const FName& /*NewName*/, FText& /*OutFailureReason*/);
+
+	SLATE_BEGIN_ARGS(STypefaceEntryEditor)
+		: _CompositeFontEditor(nullptr)
+		, _TypefaceEntry()
+		, _OnDeleteFont()
+		, _OnVerifyFontName()
+		{}
+
+		/** Pointer back to the composite font editor that owns us */
+		SLATE_ARGUMENT(SCompositeFontEditor*, CompositeFontEditor)
+
+		/** Typeface entry to edit (may be invalid, or change in response to an undo/redo) */
+		SLATE_ARGUMENT(FTypefaceListViewEntryPtr, TypefaceEntry)
+
+		/** Called when this typeface entry should be deleted */
+		SLATE_EVENT(FOnDeleteFont, OnDeleteFont)
+
+		/** Called when this typeface entry name is changed - used to verify that the new name is valid */
+		SLATE_EVENT(FOnVerifyFontName, OnVerifyFontName)
+
+	SLATE_END_ARGS()
+
+	/** Destructor */
+	~STypefaceEntryEditor();
+
+	/** SCompoundWidget interface */
+	void Construct(const FArguments& InArgs);
+
+private:
+	/** Entries used by the hinting type combo box */
+	struct FFontHintingComboEntry
+	{
+		FText DisplayName;
+		FText Tooltip;
+		EFontHinting EnumValue;
+	};
+
+	/** Get the current name of this typeface entry */
+	FText GetTypefaceEntryName() const;
+
+	/** Set the current name of this typeface entry */
+	void OnTypefaceEntryNameCommitted(const FText& InNewName, ETextCommit::Type InCommitType);
+
+	/** Verify the given typename entry name is valid */
+	bool OnTypefaceEntryChanged(const FText& InNewName, FText& OutFailureReason) const;
+
+	/** Get the current font filename of this typeface entry */
+	FString GetTypefaceEntryFontFilePath() const;
+
+	/** Set the current font filename (and associated font data) */
+	void OnTypefaceEntryFontPathPicked(const FString& InNewFontFilename);
+
+	/** Called in response to the "Delete Font" button being clicked" */
+	FReply OnDeleteFontClicked();
+
+	/** Get the current font style to use for the preview text */
+	FSlateFontInfo GetPreviewFontStyle() const;
+
+	/** Populate HintingComboData with the names of the hinting enum entries */
+	void GatherHintingEnumEntries();
+
+	/** Called when the selection of the hinting combo is changed */
+	void OnHintingComboSelectionChanged(TSharedPtr<FFontHintingComboEntry> InNewSelection, ESelectInfo::Type);
+
+	/** Make the widget for an entry in the hinting combo */
+	TSharedRef<SWidget> MakeHintingComboEntryWidget(TSharedPtr<FFontHintingComboEntry> InHintingEntry);
+
+	/** Get the text to use for the hinting combo button */
+	FText GetHintingComboText() const;
+
+	/** Pointer back to the composite font editor that owns us */
+	SCompositeFontEditor* CompositeFontEditorPtr;
+
+	/** Typeface entry to edit (may be invalid, or change in response to an undo/redo) */
+	FTypefaceListViewEntryPtr TypefaceEntry;
+
+	/** Holds a delegate that is executed when this typeface entry should be deleted */
+	FOnDeleteFont OnDeleteFont;
+
+	/** Holds a delegate that is executed when this typeface entry name is changed - used to verify that the new name is valid */
+	FOnVerifyFontName OnVerifyFontName;
+
+	/** Text block showing the preview text for this entry */
+	TSharedPtr<STextBlock> PreviewTextBlock;
+
+	/** Source data for the hinting combo widget */
+	TArray<TSharedPtr<FFontHintingComboEntry>> HintingComboData;
+
+	/** The cached text for the currently selected hinting enum entry */
+	FText ActiveHintingEnumEntryText;
+};
+
+class SSubTypefaceEditor : public SCompoundWidget
+{
+public:
+	/** Declares a delegate that is executed when this sub-typeface entry should be deleted */
+	DECLARE_DELEGATE_OneParam(FOnDeleteSubFontFamily, const FSubTypefaceListViewEntryPtr& /*SubTypefaceEntryToRemove*/);
+
+	SLATE_BEGIN_ARGS(SSubTypefaceEditor)
+		: _CompositeFontEditor(nullptr)
+		, _SubTypeface()
+		, _ParentTypeface(nullptr)
+		, _TypefaceDisplayName()
+		, _OnDeleteSubFontFamily()
+		{}
+
+		/** Pointer back to the composite font editor that owns us */
+		SLATE_ARGUMENT(SCompositeFontEditor*, CompositeFontEditor)
+
+		/** Sub-typeface to edit (may be invalid, or change in response to an undo/redo) */
+		SLATE_ARGUMENT(FSubTypefaceListViewEntryPtr, SubTypeface)
+
+		/** Parent typeface to inherit font slots from (may be invalid, or change in response to an undo/redo) */
+		SLATE_ATTRIBUTE(const FTypeface*, ParentTypeface)
+
+		/** Name to show in the UI for this typeface */
+		SLATE_ATTRIBUTE(FText, TypefaceDisplayName)
+
+		/** Called when this typeface entry should be deleted */
+		SLATE_EVENT(FOnDeleteSubFontFamily, OnDeleteSubFontFamily)
+
+	SLATE_END_ARGS()
+
+	/** Destructor */
+	~SSubTypefaceEditor();
+
+	/** SCompoundWidget interface */
+	void Construct(const FArguments& InArgs);
+
+private:
+	/** Get the typeface used by this sub-typeface */
+	FTypeface* GetTypeface() const;
+
+	/** Get the visibility of the "Add Font Override" combo button */
+	EVisibility GetAddFontOverrideVisibility() const;
+
+	/** Called before the font override combo is opened - used to update the list of available font overrides */
+	void OnAddFontOverrideComboOpening();
+
+	/** Called when the selection of the font override combo is changed */
+	void OnAddFontOverrideSelectionChanged(TSharedPtr<FName> InNewSelection, ESelectInfo::Type);
+
+	/** Make the widget for an entry in the font override combo */
+	TSharedRef<SWidget> MakeAddFontOverrideWidget(TSharedPtr<FName> InFontEntry);
+
+	/** Called in response to the "Delete Sub-Font Family" button being clicked" */
+	FReply OnDeleteSubFontFamilyClicked();
+
+	/** Update the list of character ranges for this sub-font */
+	void UpdateCharacterRangesList();
+
+	/** Make the widget for an entry in the character ranges tile view */
+	TSharedRef<ITableRow> MakeCharacterRangesEntryWidget(FCharacterRangeTileViewEntryPtr InCharacterRangeEntry, const TSharedRef<STableViewBase>& OwnerTable);
+
+	/** Called in response to the "Add Character Range" button being clicked" */
+	FReply OnAddCharacterRangeClicked();
+
+	/** Called in response to the "Delete Character Range" button being clicked" */
+	FReply OnDeleteCharacterRangeClicked(FCharacterRangeTileViewEntryPtr InCharacterRangeEntry);
+
+	/** Get the scaling factor in its numeric form */
+	TOptional<float> GetScalingFactorAsOptional() const;
+
+	/** Set the the scaling factor from its numerical form */
+	void OnScalingFactorCommittedAsNumeric(float InNewValue, ETextCommit::Type InCommitType);
+
+	/** Pointer back to the composite font editor that owns us */
+	SCompositeFontEditor* CompositeFontEditorPtr;
+
+	/** Sub-typeface to edit (may be invalid, or change in response to an undo/redo) */
+	FSubTypefaceListViewEntryPtr SubTypeface;
+
+	/** Parent typeface to inherit font slots from (may be invalid, or change in response to an undo/redo) */
+	TAttribute<const FTypeface*> ParentTypeface;
+
+	/** Called when this typeface entry should be deleted */
+	FOnDeleteSubFontFamily OnDeleteSubFontFamily;
+
+	/** Nested typeface editor widget */
+	TSharedPtr<STypefaceEditor> TypefaceEditor;
+
+	/** Font override combo box widget */
+	TSharedPtr<SComboBox<TSharedPtr<FName>>> FontOverrideCombo;
+
+	/** Source data for the font override combo widget */
+	TArray<TSharedPtr<FName>> FontOverrideComboData;
+
+	/** Internal list of sub-typeface pointers for the list view (generated from SubTypeface->CharacterRanges) */
+	TArray<FCharacterRangeTileViewEntryPtr> CharacterRangeEntries;
+
+	/** Tile view widget showing the character range editors (uses CharacterRangeEntries as its source) */
+	TSharedPtr<STileView<FCharacterRangeTileViewEntryPtr>> CharacterRangeEntriesTileView;
+};
+
+class SCharacterRangeEditor : public SCompoundWidget
+{
+public:
+	SLATE_BEGIN_ARGS(SCharacterRangeEditor)
+		: _CompositeFontEditor(nullptr)
+		, _CharacterRange()
+		{}
+
+		/** Pointer back to the composite font editor that owns us */
+		SLATE_ARGUMENT(SCompositeFontEditor*, CompositeFontEditor)
+
+		/** Character range to edit (may be invalid, or change in response to an undo/redo) */
+		SLATE_ARGUMENT(FCharacterRangeTileViewEntryPtr, CharacterRange)
+
+	SLATE_END_ARGS()
+
+	/** Destructor */
+	~SCharacterRangeEditor();
+
+	/** SCompoundWidget interface */
+	void Construct(const FArguments& InArgs);
+
+private:
+	/** Get the given range component in its TCHAR form (0 for min, 1 for max) */
+	FText GetRangeComponentAsTCHAR(int32 InComponentIndex) const;
+
+	/** Get the given range component in its numerical hex form (0 for min, 1 for max) */
+	FText GetRangeComponentAsHexString(int32 InComponentIndex) const;
+
+	/** Get the given range component in its numeric form (0 for min, 1 for max) */
+	TOptional<int32> GetRangeComponentAsOptional(int32 InComponentIndex) const;
+
+	/** Get the given range component in its numeric form (0 for min, 1 for max) */
+	int32 GetRangeComponent(const int32 InComponentIndex) const;
+
+	/** Set the the given range component from its TCHAR form (0 for min, 1 for max) */
+	void OnRangeComponentCommittedAsTCHAR(const FText& InNewValue, ETextCommit::Type InCommitType, int32 InComponentIndex);
+
+	/** Set the the given range component from its numerical hex form (0 for min, 1 for max) */
+	void OnRangeComponentCommittedAsHexString(const FText& InNewValue, ETextCommit::Type InCommitType, int32 InComponentIndex);
+
+	/** Set the the given range component from its numerical form (0 for min, 1 for max) */
+	void OnRangeComponentCommittedAsNumeric(int32 InNewValue, ETextCommit::Type InCommitType, int32 InComponentIndex);
+
+	/** Set the given range component from its numeric form (0 for min, 1 for max) */
+	void SetRangeComponent(const int32 InNewValue, const int32 InComponentIndex);
+
+	/** Pointer back to the composite font editor that owns us */
+	SCompositeFontEditor* CompositeFontEditorPtr;
+
+	/** Character range to edit (may be invalid, or change in response to an undo/redo) */
+	FCharacterRangeTileViewEntryPtr CharacterRange;
+};
diff --git a/Engine/Source/Editor/FontEditor/Public/IFontEditor.h b/Engine/Source/Editor/FontEditor/Public/IFontEditor.h
--- a/Engine/Source/Editor/FontEditor/Public/IFontEditor.h
+++ b/Engine/Source/Editor/FontEditor/Public/IFontEditor.h
@@ -17,6 +17,9 @@ public:
 
 	/** Assigns a font texture object to the page properties control when a new page is selected */
 	virtual void SetSelectedPage(int32 PageIdx) = 0;
+
+	/** Refresh the preview viewport */
+	virtual void RefreshPreview() = 0;
 };
 
 
diff --git a/Engine/Source/Editor/LandscapeEditor/Private/LandscapeEditorDetailWidgets.cpp b/Engine/Source/Editor/LandscapeEditor/Private/LandscapeEditorDetailWidgets.cpp
--- a/Engine/Source/Editor/LandscapeEditor/Private/LandscapeEditorDetailWidgets.cpp
+++ b/Engine/Source/Editor/LandscapeEditor/Private/LandscapeEditorDetailWidgets.cpp
@@ -151,7 +151,7 @@ void SToolSelector::BuildMultiBlockWidget(const ISlateStyle* StyleSet, const FNa
 	const FString MetaTag = FString::Printf(TEXT("LandscapeToolButton.%s"), Label.IsSet() == true ? *Label.Get().ToString() : TEXT("NoLabel"));
 
 	static FTextBlockStyle LabelStyle = FTextBlockStyle(FEditorStyle::GetWidgetStyle< FTextBlockStyle >(FEditorStyle::Join(StyleName, ".Label"))).SetShadowOffset(FVector2D::UnitVector);
-	static FTextBlockStyle SmallTextStyle = FTextBlockStyle(LabelStyle).SetFont(LabelStyle.Font.FontName, LabelStyle.Font.Size - 1).SetColorAndOpacity(FSlateColor::UseSubduedForeground());
+	static FTextBlockStyle SmallTextStyle = FTextBlockStyle(LabelStyle).SetFontSize(LabelStyle.Font.Size - 1).SetColorAndOpacity(FSlateColor::UseSubduedForeground());
 
 	// Create the content for our button
 	TSharedRef< SWidget > ButtonContent =
diff --git a/Engine/Source/Editor/PropertyEditor/Private/PropertyCustomizationHelpers.cpp b/Engine/Source/Editor/PropertyEditor/Private/PropertyCustomizationHelpers.cpp
--- a/Engine/Source/Editor/PropertyEditor/Private/PropertyCustomizationHelpers.cpp
+++ b/Engine/Source/Editor/PropertyEditor/Private/PropertyCustomizationHelpers.cpp
@@ -356,6 +356,8 @@ void SObjectPropertyEntryBox::Construct( const FArguments& InArgs )
 				.DisplayThumbnail(bDisplayThumbnail)
 				.OnShouldFilterAsset(InArgs._OnShouldFilterAsset)
 				.AllowClear(InArgs._AllowClear)
+				.DisplayUseSelected(InArgs._DisplayUseSelected)
+				.DisplayBrowse(InArgs._DisplayBrowse)
 				.PropertyHandle(PropertyHandle)
 				.ThumbnailSize(ThumbnailSize)
 		]
diff --git a/Engine/Source/Editor/PropertyEditor/Private/UserInterface/PropertyEditor/SPropertyEditorAsset.cpp b/Engine/Source/Editor/PropertyEditor/Private/UserInterface/PropertyEditor/SPropertyEditorAsset.cpp
--- a/Engine/Source/Editor/PropertyEditor/Private/UserInterface/PropertyEditor/SPropertyEditorAsset.cpp
+++ b/Engine/Source/Editor/PropertyEditor/Private/UserInterface/PropertyEditor/SPropertyEditorAsset.cpp
@@ -73,7 +73,6 @@ void SPropertyEditorAsset::Construct( const FArguments& InArgs, const TSharedPtr
 	PropertyHandle = InArgs._PropertyHandle;
 	OnSetObject = InArgs._OnSetObject;
 
-	UObject* Object = NULL;
 	if(PropertyEditor.IsValid())
 	{
 		UProperty* NodeProperty = PropertyEditor->GetPropertyNode()->GetProperty();
@@ -81,9 +80,34 @@ void SPropertyEditorAsset::Construct( const FArguments& InArgs, const TSharedPtr
 		check(ObjectProperty);
 
 		bAllowClear = !(NodeProperty->PropertyFlags & CPF_NoClear);
-		PropertyEditor->GetPropertyHandle()->GetValue( Object );
 		ObjectClass = ObjectProperty->PropertyClass;
 		bIsActor = ObjectProperty->PropertyClass->IsChildOf( AActor::StaticClass() );
+
+		FString ClassFilterString = NodeProperty->GetMetaData("AllowedClasses");
+		if(ClassFilterString.IsEmpty())
+		{
+			CustomClassFilters.Add(ObjectClass);
+		}
+		else
+		{
+			TArray<FString> CustomClassFilterNames;
+			ClassFilterString.ParseIntoArray(&CustomClassFilterNames, TEXT(","), true);
+
+			for(auto It = CustomClassFilterNames.CreateConstIterator(); It; ++It)
+			{
+				const FString& ClassName = *It;
+
+				UClass* Class = FindObject<UClass>(ANY_PACKAGE, *ClassName);
+				if(!Class)
+				{
+					Class = LoadObject<UClass>(nullptr, *ClassName);
+				}
+				if(Class)
+				{
+					CustomClassFilters.Add(Class);
+				}
+			}
+		}
 	}
 	else
 	{
@@ -91,6 +115,8 @@ void SPropertyEditorAsset::Construct( const FArguments& InArgs, const TSharedPtr
 		ObjectPath = InArgs._ObjectPath;
 		ObjectClass = InArgs._Class;
 		bIsActor = ObjectClass->IsChildOf( AActor::StaticClass() );
+
+		CustomClassFilters.Add(ObjectClass);
 	}
 	OnShouldFilterAsset = InArgs._OnShouldFilterAsset;
 
@@ -185,7 +211,7 @@ void SPropertyEditorAsset::Construct( const FArguments& InArgs, const TSharedPtr
 		]
 	];
 
-	if(!bIsActor)
+	if(!bIsActor && InArgs._DisplayUseSelected)
 	{
 		ButtonBox->AddSlot()
 		.VAlign(VAlign_Center)
@@ -196,16 +222,19 @@ void SPropertyEditorAsset::Construct( const FArguments& InArgs, const TSharedPtr
 		];
 	}
 
-	ButtonBox->AddSlot()
-	.AutoWidth()
-	.Padding( 2.0f, 0.0f )
-	.VAlign(VAlign_Center)
-	[
-		PropertyCustomizationHelpers::MakeBrowseButton(
-			FSimpleDelegate::CreateSP( this, &SPropertyEditorAsset::OnBrowse ),
-			TAttribute<FText>( this, &SPropertyEditorAsset::GetOnBrowseToolTip )
-			)
-	];
+	if(InArgs._DisplayBrowse)
+	{
+		ButtonBox->AddSlot()
+		.AutoWidth()
+		.Padding( 2.0f, 0.0f )
+		.VAlign(VAlign_Center)
+		[
+			PropertyCustomizationHelpers::MakeBrowseButton(
+				FSimpleDelegate::CreateSP( this, &SPropertyEditorAsset::OnBrowse ),
+				TAttribute<FText>( this, &SPropertyEditorAsset::GetOnBrowseToolTip )
+				)
+		];
+	}
 
 	if(bIsActor)
 	{
@@ -308,12 +337,9 @@ TSharedRef<SWidget> SPropertyEditorAsset::OnGetMenuContent()
 	}
 	else
 	{
-		TArray<const UClass*> AllowedClasses;
-		AllowedClasses.Add(ObjectClass);
-
 		return PropertyCustomizationHelpers::MakeAssetPickerWithMenu(Value.AssetData,
 																	 bAllowClear,
-																	 AllowedClasses,
+																	 CustomClassFilters,
 																	 NewAssetFactories,
 																	 OnShouldFilterAsset,
 																	 FOnAssetSelected::CreateSP(this, &SPropertyEditorAsset::OnAssetSelected),
@@ -526,7 +552,7 @@ void SPropertyEditorAsset::OnActorSelected( AActor* InActor )
 
 void SPropertyEditorAsset::OnGetAllowedClasses(TArray<const UClass*>& AllowedClasses)
 {
-	AllowedClasses.Add(ObjectClass);
+	AllowedClasses.Append(CustomClassFilters);
 }
 
 void SPropertyEditorAsset::OnOpenAssetEditor()
diff --git a/Engine/Source/Editor/PropertyEditor/Private/UserInterface/PropertyEditor/SPropertyEditorAsset.h b/Engine/Source/Editor/PropertyEditor/Private/UserInterface/PropertyEditor/SPropertyEditorAsset.h
--- a/Engine/Source/Editor/PropertyEditor/Private/UserInterface/PropertyEditor/SPropertyEditorAsset.h
+++ b/Engine/Source/Editor/PropertyEditor/Private/UserInterface/PropertyEditor/SPropertyEditorAsset.h
@@ -18,6 +18,8 @@ public:
 		, _ClassFont( FEditorStyle::GetFontStyle("PropertyEditor.AssetClass.Font") ) 
 		, _AllowClear(true)
 		, _DisplayThumbnail(false)
+		, _DisplayUseSelected(true)
+		, _DisplayBrowse(true)
 		, _ThumbnailPool(NULL)
 		, _ThumbnailSize( FIntPoint(64, 64) )
 		, _ObjectPath()
@@ -29,6 +31,8 @@ public:
 		SLATE_ATTRIBUTE( FSlateFontInfo, ClassFont )
 		SLATE_ARGUMENT( bool, AllowClear )
 		SLATE_ARGUMENT( bool, DisplayThumbnail )
+		SLATE_ARGUMENT( bool, DisplayUseSelected )
+		SLATE_ARGUMENT( bool, DisplayBrowse )
 		SLATE_ARGUMENT( TSharedPtr<FAssetThumbnailPool>, ThumbnailPool )
 		SLATE_ARGUMENT( FIntPoint, ThumbnailSize )
 		SLATE_ATTRIBUTE( FString, ObjectPath )
@@ -252,6 +256,9 @@ private:
 	/** The class of the object we are editing */
 	UClass* ObjectClass;
 
+	/** Classes that can be used with this property */
+	TArray<const UClass*> CustomClassFilters;
+
 	/** A list of the factories we can use to create new assets */
 	TArray<UFactory*> NewAssetFactories;
 
diff --git a/Engine/Source/Editor/PropertyEditor/Public/PropertyCustomizationHelpers.h b/Engine/Source/Editor/PropertyEditor/Public/PropertyCustomizationHelpers.h
--- a/Engine/Source/Editor/PropertyEditor/Public/PropertyCustomizationHelpers.h
+++ b/Engine/Source/Editor/PropertyEditor/Public/PropertyCustomizationHelpers.h
@@ -58,6 +58,8 @@ public:
 	SLATE_BEGIN_ARGS( SObjectPropertyEntryBox )
 		: _AllowedClass( UObject::StaticClass() )
 		, _AllowClear( true )
+		, _DisplayUseSelected( true )
+		, _DisplayBrowse( true )
 	{}
 		/** The path to the object */
 		SLATE_ATTRIBUTE( FString, ObjectPath )
@@ -77,6 +79,10 @@ public:
 		SLATE_EVENT(FOnShouldFilterAsset, OnShouldFilterAsset)
 		/** Whether the asset can be 'None' */
 		SLATE_ARGUMENT(bool, AllowClear)
+		/** Whether to show the 'Use Selected' button */
+		SLATE_ARGUMENT(bool, DisplayUseSelected)
+		/** Whether to show the 'Browse' button */
+		SLATE_ARGUMENT(bool, DisplayBrowse)
 	SLATE_END_ARGS()
 
 	PROPERTYEDITOR_API void Construct( const FArguments& InArgs );
diff --git a/Engine/Source/Editor/TranslationEditor/Private/TranslationEditor.cpp b/Engine/Source/Editor/TranslationEditor/Private/TranslationEditor.cpp
--- a/Engine/Source/Editor/TranslationEditor/Private/TranslationEditor.cpp
+++ b/Engine/Source/Editor/TranslationEditor/Private/TranslationEditor.cpp
@@ -27,6 +27,20 @@ DEFINE_LOG_CATEGORY_STATIC(LocalizationExport, Log, All);
 
 #define LOCTEXT_NAMESPACE "TranslationEditor"
 
+namespace TranslationEditorUtils
+{
+
+/** Get the filename used by the given font info */
+FString GetFontFilename(const FSlateFontInfo& InFontInfo)
+{
+	const FCompositeFont* const ResolvedCompositeFont = InFontInfo.GetCompositeFont();
+	return (ResolvedCompositeFont && ResolvedCompositeFont->DefaultTypeface.Fonts.Num() > 0)
+		? ResolvedCompositeFont->DefaultTypeface.Fonts[0].Font.FontFilename
+		: "";
+}
+
+} // namespace TranslationEditorUtils
+
 const FName FTranslationEditor::UntranslatedTabId( TEXT( "TranslationEditor_Untranslated" ) );
 const FName FTranslationEditor::ReviewTabId( TEXT( "TranslationEditor_Review" ) );
 const FName FTranslationEditor::CompletedTabId( TEXT( "TranslationEditor_Completed" ) );
@@ -757,7 +771,7 @@ void FTranslationEditor::MapActions()
 void FTranslationEditor::ChangeSourceFont()
 {
 	// Use path from current font
-	FString DefaultFile(SourceFont.FontName.ToString());
+	FString DefaultFile = TranslationEditorUtils::GetFontFilename(SourceFont);
 
 	FString NewFontFilename;
 	bool bOpened = OpenFontPicker(DefaultFile, NewFontFilename);
@@ -772,7 +786,7 @@ void FTranslationEditor::ChangeSourceFont()
 void FTranslationEditor::ChangeTranslationTargetFont()
 {
 	// Use path from current font
-	FString DefaultFile(TranslationTargetFont.FontName.ToString());
+	FString DefaultFile = TranslationEditorUtils::GetFontFilename(TranslationTargetFont);
 
 	FString NewFontFilename;
 	bool bOpened = OpenFontPicker(DefaultFile, NewFontFilename);
diff --git a/Engine/Source/Editor/UMGEditor/Private/Details/SWidgetDetailsView.cpp b/Engine/Source/Editor/UMGEditor/Private/Details/SWidgetDetailsView.cpp
--- a/Engine/Source/Editor/UMGEditor/Private/Details/SWidgetDetailsView.cpp
+++ b/Engine/Source/Editor/UMGEditor/Private/Details/SWidgetDetailsView.cpp
@@ -11,8 +11,9 @@
 #include "DetailWidgetExtensionHandler.h"
 #include "EditorClassUtils.h"
 
-#include "Customizations/SlateBrushCustomization.h"
 #include "Customizations/DetailCustomizations.h"
+#include "Customizations/SlateBrushCustomization.h"
+#include "Customizations/SlateFontInfoCustomization.h"
 
 #include "WidgetNavigationCustomization.h"
 #include "CanvasSlotCustomization.h"
@@ -168,6 +169,7 @@ SWidgetDetailsView::~SWidgetDetailsView()
 	PropertyModule.UnregisterCustomPropertyTypeLayout(TEXT("EVerticalAlignment"), nullptr, PropertyView);
 	PropertyModule.UnregisterCustomPropertyTypeLayout(TEXT("SlateChildSize"), nullptr, PropertyView);
 	PropertyModule.UnregisterCustomPropertyTypeLayout(TEXT("SlateBrush"), nullptr, PropertyView);
+	PropertyModule.UnregisterCustomPropertyTypeLayout(TEXT("SlateFontInfo"), nullptr, PropertyView);
 	PropertyModule.UnregisterCustomPropertyTypeLayout(TEXT("ETextJustify"), nullptr, PropertyView);
 }
 
@@ -184,6 +186,7 @@ void SWidgetDetailsView::RegisterCustomizations()
 	PropertyModule.RegisterCustomPropertyTypeLayout(TEXT("EVerticalAlignment"), FOnGetPropertyTypeCustomizationInstance::CreateStatic(&FVerticalAlignmentCustomization::MakeInstance), nullptr, PropertyView);
 	PropertyModule.RegisterCustomPropertyTypeLayout(TEXT("SlateChildSize"), FOnGetPropertyTypeCustomizationInstance::CreateStatic(&FSlateChildSizeCustomization::MakeInstance), nullptr, PropertyView);
 	PropertyModule.RegisterCustomPropertyTypeLayout(TEXT("SlateBrush"), FOnGetPropertyTypeCustomizationInstance::CreateStatic(&FSlateBrushStructCustomization::MakeInstance, false), nullptr, PropertyView);
+	PropertyModule.RegisterCustomPropertyTypeLayout(TEXT("SlateFontInfo"), FOnGetPropertyTypeCustomizationInstance::CreateStatic(&FSlateFontInfoStructCustomization::MakeInstance), nullptr, PropertyView);
 	PropertyModule.RegisterCustomPropertyTypeLayout(TEXT("ETextJustify"), FOnGetPropertyTypeCustomizationInstance::CreateStatic(&FTextJustifyCustomization::MakeInstance), nullptr, PropertyView);
 }
 
diff --git a/Engine/Source/Editor/UnrealEd/Classes/Factories/FontFactory.h b/Engine/Source/Editor/UnrealEd/Classes/Factories/FontFactory.h
--- a/Engine/Source/Editor/UnrealEd/Classes/Factories/FontFactory.h
+++ b/Engine/Source/Editor/UnrealEd/Classes/Factories/FontFactory.h
@@ -7,23 +7,13 @@
 #pragma once
 #include "FontFactory.generated.h"
 
-UENUM()
-enum EFontCharacterSet
-{
-	CHARSET_Default,
-	CHARSET_Ansi,
-	CHARSET_Symbol,
-	CHARSET_MAX,
-};
-
 UCLASS()
-class UFontFactory : public UTextureFactory
+class UNREALED_API UFontFactory : public UFactory
 {
 	GENERATED_UCLASS_BODY()
 
-
 	// Begin UFactory Interface
-	virtual UObject* FactoryCreateBinary( UClass* Class, UObject* InParent, FName Name, EObjectFlags Flags, UObject* Context, const TCHAR* Type, const uint8*& Buffer, const uint8* BufferEnd, FFeedbackContext* Warn ) override;
+	virtual UObject* FactoryCreateNew(UClass* InClass, UObject* InParent, FName InName, EObjectFlags Flags, UObject* Context, FFeedbackContext* Warn) override;
 	// Begin UFactory Interface	
 };
 
diff --git a/Engine/Source/Editor/UnrealEd/Classes/Factories/FontFileImportFactory.h b/Engine/Source/Editor/UnrealEd/Classes/Factories/FontFileImportFactory.h
new file mode 100644
--- /dev/null
+++ b/Engine/Source/Editor/UnrealEd/Classes/Factories/FontFileImportFactory.h
@@ -0,0 +1,21 @@
+// Copyright 1998-2014 Epic Games, Inc. All Rights Reserved.
+
+//=============================================================================
+// FontFactory: Creates a Font Factory
+//=============================================================================
+
+#pragma once
+#include "FontFileImportFactory.generated.h"
+
+UCLASS()
+class UNREALED_API UFontFileImportFactory : public UFactory
+{
+	GENERATED_UCLASS_BODY()
+
+	// Begin UFactory Interface
+	virtual UObject* FactoryCreateBinary(UClass* Class, UObject* InParent, FName Name, EObjectFlags Flags, UObject* Context, const TCHAR* Type, const uint8*& Buffer, const uint8* BufferEnd, FFeedbackContext* Warn) override;
+	// Begin UFactory Interface	
+};
+
+
+
diff --git a/Engine/Source/Editor/UnrealEd/Classes/Factories/TrueTypeFontFactory.h b/Engine/Source/Editor/UnrealEd/Classes/Factories/TrueTypeFontFactory.h
--- a/Engine/Source/Editor/UnrealEd/Classes/Factories/TrueTypeFontFactory.h
+++ b/Engine/Source/Editor/UnrealEd/Classes/Factories/TrueTypeFontFactory.h
@@ -8,7 +8,7 @@
 #include "TrueTypeFontFactory.generated.h"
 
 UCLASS(hidecategories=Object, collapsecategories)
-class UTrueTypeFontFactory : public UFontFactory, public FReimportHandler
+class UNREALED_API UTrueTypeFontFactory : public UTextureFactory, public FReimportHandler
 {
 	GENERATED_UCLASS_BODY()
 
@@ -31,6 +31,11 @@ class UTrueTypeFontFactory : public UFontFactory, public FReimportHandler
 	// Begin UFactory Interface
 	virtual bool ConfigureProperties() override;
 	virtual UObject* FactoryCreateNew(UClass* Class,UObject* InParent,FName Name,EObjectFlags Flags,UObject* Context,FFeedbackContext* Warn) override;
+	virtual bool ShouldShowInNewMenu() const override
+	{
+		// Don't show this factory in the content browser menu; it's invoked manually when changing the UFont cache type to "Offline" 
+		return false;
+	}
 	// Begin UFactory Interface	
 
 	// Begin FReimportHandler interface
diff --git a/Engine/Source/Editor/UnrealEd/Private/Factories/EditorFactories.cpp b/Engine/Source/Editor/UnrealEd/Private/Factories/EditorFactories.cpp
--- a/Engine/Source/Editor/UnrealEd/Private/Factories/EditorFactories.cpp
+++ b/Engine/Source/Editor/UnrealEd/Private/Factories/EditorFactories.cpp
@@ -5074,196 +5074,60 @@ bool UTextureExporterTGA::ExportBinary( UObject* Object, const TCHAR* Type, FArc
 	UFontFactory.
 ------------------------------------------------------------------------------*/
 
-//
-//	Fast pixel-lookup.
-//
-static inline uint8 AT( const uint8* Screen, uint32 SXL, uint32 X, uint32 Y )
-{
-	return Screen[X+Y*SXL];
-}
-
-//
-// Codepage 850 -> Latin-1 mapping table:
-//
-uint8 FontRemap[256] = 
-{
-	0,  1,  2,  3,  4,  5,  6,  7,  8,  9, 10, 11, 12, 13, 14, 15,
-	16, 17, 18, 19, 20, 21, 22, 23, 24, 25, 26, 27, 28, 29, 30, 31,
-	32, 33, 34, 35, 36, 37, 38, 39, 40, 41, 42, 43, 44, 45, 46, 47,
-	48, 49, 50, 51, 52, 53, 54, 55, 56, 57, 58, 59, 60, 61, 62, 63,
-
-	64, 65, 66, 67, 68, 69, 70, 71, 72, 73, 74, 75, 76, 77, 78, 79,
-	80, 81, 82, 83, 84, 85, 86, 87, 88, 89, 90, 91, 92, 93, 94, 95,
-	96, 97, 98, 99,100,101,102,103,104,105,106,107,108,109,110,111,
-	112,113,114,115,116,117,118,119,120,121,122,123,124,125,126,127,
-
-	000,000,000,000,000,000,000,000,000,000,000,000,000,000,000,000,
-	000,000,000,000,000,000,000,000,000,000,000,000,000,000,000,000,
-	032,173,184,156,207,190,124,245,034,184,166,174,170,196,169,238,
-	248,241,253,252,239,230,244,250,247,251,248,175,172,171,243,168,
-
-	183,181,182,199,142,143,146,128,212,144,210,211,222,214,215,216,
-	209,165,227,224,226,229,153,158,157,235,233,234,154,237,231,225,
-	133,160,131,196,132,134,145,135,138,130,136,137,141,161,140,139,
-	208,164,149,162,147,228,148,246,155,151,163,150,129,236,232,152,
-};
-
-//
-//	Find the border around a font glyph that starts at x,y (it's upper
-//	left hand corner).  If it finds a glyph box, it returns 0 and the
-//	glyph 's length (xl,yl).  Otherwise returns -1.
-//
-static bool ScanFontBox( const uint8* Data, int32 X, int32 Y, int32& XL, int32& YL, int32 SizeX )
+UFontFactory::UFontFactory(const FObjectInitializer& ObjectInitializer)
+	: Super(ObjectInitializer)
 {
-	int32 FontXL = SizeX;
-
-	// Find x-length.
-	int32 NewXL = 1;
-	while ( AT(Data,FontXL,X+NewXL,Y)==255 && AT(Data,FontXL,X+NewXL,Y+1)!=255 )
-	{
-		NewXL++;
-	}
-
-	if( AT(Data,FontXL,X+NewXL,Y)!=255 )
-	{
-		return 0;
-	}
+	SupportedClass = UFont::StaticClass();
 
-	// Find y-length.
-	int32 NewYL = 1;
-	while( AT(Data,FontXL,X,Y+NewYL)==255 && AT(Data,FontXL,X+1,Y+NewYL)!=255 )
-	{
-		NewYL++;
-	}
+	bCreateNew = true;
+	bEditAfterNew = true;
+}
 
-	if( AT(Data,FontXL,X,Y+NewYL)!=255 )
+UObject* UFontFactory::FactoryCreateNew(UClass* InClass, UObject* InParent, FName InName, EObjectFlags InFlags, UObject* InContext, FFeedbackContext* InWarn)
+{
+	UFont* const Font = ConstructObject<UFont>(InClass, InParent, InName, InFlags);
+	if(Font)
 	{
-		return 0;
+		Font->FontCacheType = EFontCacheType::Runtime;
 	}
-
-	XL = NewXL - 1;
-	YL = NewYL - 1;
-
-	return 1;
+	return Font;
 }
 
-#define NUM_FONT_CHARS 256
+/*------------------------------------------------------------------------------
+	UFontFileImportFactory.
+------------------------------------------------------------------------------*/
 
-UFontFactory::UFontFactory(const FObjectInitializer& ObjectInitializer)
+UFontFileImportFactory::UFontFileImportFactory(const FObjectInitializer& ObjectInitializer)
 	: Super(ObjectInitializer)
 {
 	SupportedClass = UFont::StaticClass();
-	bEditorImport = false;
-	LODGroup = TEXTUREGROUP_UI;
-}
 
-UObject* UFontFactory::FactoryCreateBinary
-(
-	UClass*				Class,
-	UObject*			InParent,
-	FName				Name,
-	EObjectFlags		Flags,
-	UObject*			Context,
-	const TCHAR*		Type,
-	const uint8*&		Buffer,
-	const uint8*			BufferEnd,
-	FFeedbackContext*	Warn
-)
-{
-	FEditorDelegates::OnAssetPreImport.Broadcast(this, Class, InParent, Name, Type);
+	bEditorImport = true;
 
-	check(Type);
+	Formats.Add(TEXT("ttf;TrueType font"));
+	Formats.Add(TEXT("otf;OpenType font"));
+}
 
-	check(Class==UFont::StaticClass());
-	UFont* Font = new( InParent, Name, Flags )UFont(FObjectInitializer());
-	// note RF_Public because font textures can be referenced direclty by material expressions
-	UTexture2D* Tex = CastChecked<UTexture2D>( UTextureFactory::FactoryCreateBinary( 
-		UTexture2D::StaticClass(), Font, NAME_None, RF_Public, Context, Type, Buffer, BufferEnd, Warn ) );
+UObject* UFontFileImportFactory::FactoryCreateBinary(UClass* InClass, UObject* InParent, FName InName, EObjectFlags InFlags, UObject* InContext, const TCHAR* InType, const uint8*& InBuffer, const uint8* InBufferEnd, FFeedbackContext* InWarn)
+{
+	FEditorDelegates::OnAssetPreImport.Broadcast(this, InClass, InParent, InName, InType);
 
-	if( Tex != nullptr )
+	UFont* const Font = ConstructObject<UFont>(InClass, InParent, InName, InFlags);
+	if(Font)
 	{
-		Tex->LODGroup = TEXTUREGROUP_UI;  // set the LOD group otherwise this will be in the World Group
-
-		// Also, we never want to stream in font textures since that always looks awful
-		Tex->NeverStream = true;
+		Font->FontCacheType = EFontCacheType::Runtime;
 
-		Font->Textures.Add(Tex);
+		FString FontFilename = GetCurrentFilename();
 
-		// Init.
-		uint8* TextureData = (uint8*)Tex->Source.LockMip(0);
-		int32 TexSizeX = Tex->Source.GetSizeX();
-		int32 TexSizeY = Tex->Source.GetSizeY();
-		Font->Characters.AddZeroed( NUM_FONT_CHARS );
-
-		// Scan in all fonts, starting at glyph 32.
-		uint32 i = 32;
-		int32 Y = 0;
-		do
-		{
-			int32 X = 0;
-			while( AT(TextureData,TexSizeX,X,Y)!=255 && Y<TexSizeY )
-			{
-				X++;
-				if( X >= TexSizeX )
-				{
-					X = 0;
-					if( ++Y >= TexSizeY )
-						break;
-				}
-			}
+		TArray<uint8> FontData;
+		FontData.Insert(InBuffer, InBufferEnd - InBuffer, 0);
 
-			// Scan all glyphs in this row.
-			if( Y < TexSizeY )
-			{
-				int32 XL=0, YL=0, MaxYL=0;
-				while( i<(uint32)Font->Characters.Num() && ScanFontBox(TextureData,X,Y,XL,YL,TexSizeX) )
-				{
-					Font->Characters[i].StartU = X+1;
-					Font->Characters[i].StartV = Y+1;
-					Font->Characters[i].USize  = XL;
-					Font->Characters[i].VSize  = YL;
-					Font->Characters[i].TextureIndex = 0;
-					Font->Characters[i].VerticalOffset = 0;
-					X += XL + 1;
-					i++;
-					if( YL > MaxYL )
-						MaxYL = YL;
-				}
-				Y += MaxYL + 1;
-			}
-		} while( i<(uint32)Font->Characters.Num() && Y<TexSizeY );
-
-		// Cleanup font data.
-		int32 MipBytes = Tex->Source.CalcMipSize(0);
-		for( int32 j=0; j<MipBytes; j++ )
-		{
-			if( TextureData[j]==255 )
-			{
-				TextureData[j] = 0;
-			}
-		}
-		Tex->Source.UnlockMip(0);
-		Tex->PostEditChange();
-
-		// Remap old fonts.
-		TArray<FFontCharacter> Old = Font->Characters;
-		for( i=0; i<(uint32)Font->Characters.Num(); i++ )
-		{
-			Font->Characters[i] = Old[FontRemap[i]];
-		}
-
-		Font->CacheCharacterCountAndMaxCharHeight();
-
-		FEditorDelegates::OnAssetPostImport.Broadcast(this, Font);
-
-		return Font;
-	}
-	else 
-	{
-		Font->MarkPendingKill();
-		FEditorDelegates::OnAssetPostImport.Broadcast( this, nullptr );
-		return nullptr;
+		Font->CompositeFont.DefaultTypeface.Fonts.Add(FTypefaceEntry("Default", MoveTemp(FontFilename), MoveTemp(FontData), EFontHinting::Auto));
 	}
+
+	FEditorDelegates::OnAssetPostImport.Broadcast(this, Font);
+	
+	return Font;
 }
 
 /*------------------------------------------------------------------------------
diff --git a/Engine/Source/Editor/UnrealEd/Private/Factories/TTFontImport.cpp b/Engine/Source/Editor/UnrealEd/Private/Factories/TTFontImport.cpp
--- a/Engine/Source/Editor/UnrealEd/Private/Factories/TTFontImport.cpp
+++ b/Engine/Source/Editor/UnrealEd/Private/Factories/TTFontImport.cpp
@@ -42,6 +42,7 @@ UTrueTypeFontFactory::UTrueTypeFontFactory(const FObjectInitializer& ObjectIniti
 	bCreateNew = true;
 	bEditAfterNew = true;
 	AutoPriority = -1;
+	LODGroup = TEXTUREGROUP_UI;
 }
 
 void UTrueTypeFontFactory::PostInitProperties()
diff --git a/Engine/Source/Editor/UnrealEd/Private/ThumbnailRendering/FontThumbnailRenderer.cpp b/Engine/Source/Editor/UnrealEd/Private/ThumbnailRendering/FontThumbnailRenderer.cpp
--- a/Engine/Source/Editor/UnrealEd/Private/ThumbnailRendering/FontThumbnailRenderer.cpp
+++ b/Engine/Source/Editor/UnrealEd/Private/ThumbnailRendering/FontThumbnailRenderer.cpp
@@ -7,38 +7,85 @@ UFontThumbnailRenderer::UFontThumbnailRenderer(const FObjectInitializer& ObjectI
 {
 }
 
-void UFontThumbnailRenderer::GetThumbnailSize(UObject* Object, float Zoom,uint32& OutWidth,uint32& OutHeight) const
+void UFontThumbnailRenderer::GetThumbnailSize(UObject* Object, float Zoom, uint32& OutWidth, uint32& OutHeight) const
 {
-	UFont* Font = Cast<UFont>(Object);
-	if (Font != nullptr && 
-		Font->Textures.Num() > 0 &&
-		Font->Textures[0] != nullptr)
-	{
-		// Get the texture interface for the font text
-		UTexture2D* Tex = Font->Textures[0];
-		OutWidth = FMath::TruncToInt(Zoom * (float)Tex->GetSurfaceWidth());
-		OutHeight = FMath::TruncToInt(Zoom * (float)Tex->GetSurfaceHeight());
-	}
-	else
+	UFont* const Font = Cast<UFont>(Object);
+
+	if(Font)
 	{
-		OutWidth = OutHeight = 0;
+		switch(Font->FontCacheType)
+		{
+		case EFontCacheType::Offline:
+			if(Font->Textures.Num() > 0 && Font->Textures[0])
+			{
+				// Get the texture interface for the font text
+				UTexture2D* const Tex = Font->Textures[0];
+				OutWidth = FMath::TruncToInt(Zoom * (float)Tex->GetSurfaceWidth());
+				OutHeight = FMath::TruncToInt(Zoom * (float)Tex->GetSurfaceHeight());
+				return;
+			}
+			break;
+
+		case EFontCacheType::Runtime:
+			if(Font->CompositeFont.DefaultTypeface.Fonts.Num() > 0)
+			{
+				OutWidth = OutHeight = Zoom * 256.0f;
+			}
+			break;
+
+		default:
+			break;
+		};
 	}
+
+	OutWidth = OutHeight = 0;
 }
 
 void UFontThumbnailRenderer::Draw(UObject* Object, int32 X, int32 Y, uint32 Width, uint32 Height, FRenderTarget*, FCanvas* Canvas)
 {
-	UFont* Font = Cast<UFont>(Object);
-	if (Font != nullptr && 
-		Font->Textures.Num() > 0 &&
-		Font->Textures[0] != nullptr)
+	UFont* const Font = Cast<UFont>(Object);
+
+	if(Font)
 	{
-		FCanvasTileItem TileItem( FVector2D( X, Y ), Font->Textures[0]->Resource, FLinearColor::White );
-		TileItem.BlendMode = SE_BLEND_Translucent;
-		if (Font->ImportOptions.bUseDistanceFieldAlpha)
+		switch(Font->FontCacheType)
 		{
-			TileItem .BlendMode = SE_BLEND_TranslucentDistanceField;
-		}
+		case EFontCacheType::Offline:
+			if(Font->Textures.Num() > 0 && Font->Textures[0])
+			{
+				FCanvasTileItem TileItem(FVector2D(X, Y), Font->Textures[0]->Resource, FLinearColor::White);
+				TileItem.BlendMode = (Font->ImportOptions.bUseDistanceFieldAlpha) ? SE_BLEND_TranslucentDistanceField : SE_BLEND_Translucent;
+				Canvas->DrawItem(TileItem);
+			}
+			break;
 
-		Canvas->DrawItem( TileItem );		
+		case EFontCacheType::Runtime:
+			if(Font->CompositeFont.DefaultTypeface.Fonts.Num() > 0)
+			{
+				static const int32 FontSize = 28;
+				static const float FontScale = 1.0f;
+
+				const FText FontNameText = FText::FromName(Object->GetFName());
+				TSharedRef<FSlateFontCache> FontCache = FSlateApplication::Get().GetRenderer()->GetFontCache();
+
+				// Draw the object name using each of the default fonts
+				FVector2D CurPos(X, Y);
+				for(const FTypefaceEntry& TypefaceEntry : Font->CompositeFont.DefaultTypeface.Fonts)
+				{
+					FSlateFontInfo FontInfo(Font, FontSize, TypefaceEntry.Name);
+
+					FCharacterList& CharacterList = FontCache->GetCharacterList(FontInfo, FontScale);
+					const int32 MaxCharHeight = CharacterList.GetMaxHeight();
+
+					FCanvasTextItem TextItem(CurPos, FontNameText, FontInfo, FLinearColor::White);
+					Canvas->DrawItem(TextItem);
+
+					CurPos.Y += MaxCharHeight;
+				}
+			}
+			break;
+
+		default:
+			break;
+		};
 	}
-}
\ No newline at end of file
+}
diff --git a/Engine/Source/Runtime/Core/Public/UObject/ObjectVersion.h b/Engine/Source/Runtime/Core/Public/UObject/ObjectVersion.h
--- a/Engine/Source/Runtime/Core/Public/UObject/ObjectVersion.h
+++ b/Engine/Source/Runtime/Core/Public/UObject/ObjectVersion.h
@@ -628,6 +628,8 @@ enum EUnrealEngineObjectUE4Version
 	VER_UE4_BP_MATH_VECTOR_EQUALITY_USES_EPSILON,
 	// Static lighting support was re-added to foliage, and mobility was returned to static
 	VER_UE4_FOLIAGE_STATIC_LIGHTING_SUPPORT,
+	// Added composite fonts to Slate font info
+	VER_UE4_SLATE_COMPOSITE_FONTS,
 
 	// -----<new versions can be added before this line>-------------------------------------------------
 	// - this needs to be the last line (see note below)
diff --git a/Engine/Source/Runtime/Engine/Classes/Engine/Canvas.h b/Engine/Source/Runtime/Engine/Classes/Engine/Canvas.h
--- a/Engine/Source/Runtime/Engine/Classes/Engine/Canvas.h
+++ b/Engine/Source/Runtime/Engine/Classes/Engine/Canvas.h
@@ -11,6 +11,7 @@
 class FSceneView;
 class FCanvas;
 class FCanvasItem;
+class UFont;
 
 /**
  * Holds texture information with UV coordinates as well.
@@ -80,7 +81,7 @@ struct FTextSizingParameters
 
 	/** the font to use for sizing/wrapping the string */
 	UPROPERTY()
-	class UFont* DrawFont;
+	const UFont* DrawFont;
 
 	/** Horizontal spacing adjustment between characters and vertical spacing adjustment between wrapped lines */
 	UPROPERTY()
@@ -98,13 +99,13 @@ struct FTextSizingParameters
 	}
 
 
-		FTextSizingParameters( float inDrawX, float inDrawY, float inDrawXL, float inDrawYL, UFont* inFont=NULL )
+		FTextSizingParameters( float inDrawX, float inDrawY, float inDrawXL, float inDrawYL, const UFont* inFont=NULL )
 		: DrawX(inDrawX), DrawY(inDrawY), DrawXL(inDrawXL), DrawYL(inDrawYL)
 		, Scaling(1.f,1.f), DrawFont(inFont)
 		, SpacingAdjust( 0.0f, 0.0f )
 		{
 		}
-		FTextSizingParameters( UFont* inFont, float ScaleX, float ScaleY)
+		FTextSizingParameters( const UFont* inFont, float ScaleX, float ScaleY)
 		: DrawX(0.f), DrawY(0.f), DrawXL(0.f), DrawYL(0.f)
 		, Scaling(ScaleX,ScaleY), DrawFont(inFont)
 		, SpacingAdjust( 0.0f, 0.0f )
@@ -256,18 +257,18 @@ public:
 	 * @param YL out - vertical length of string
 	 * @param Text - string to calculate for
 	 */
-	static void ClippedStrLen(UFont* Font, float ScaleX, float ScaleY, int32& XL, int32& YL, const TCHAR* Text);
+	static void ClippedStrLen(const UFont* Font, float ScaleX, float ScaleY, int32& XL, int32& YL, const TCHAR* Text);
 
 	/**	
 	 * Calculate the size of a string built from a font, word wrapped
 	 * to a specified region.
 	 */
-	void VARARGS WrappedStrLenf(UFont* Font, float ScaleX, float ScaleY, int32& XL, int32& YL, const TCHAR* Fmt, ...);
+	void VARARGS WrappedStrLenf(const UFont* Font, float ScaleX, float ScaleY, int32& XL, int32& YL, const TCHAR* Fmt, ...);
 		
 	/**
 	 * Compute size and optionally print text with word wrap.
 	 */
-	int32 WrappedPrint(bool Draw, float X, float Y, int32& out_XL, int32& out_YL, UFont* Font, float ScaleX, float ScaleY, bool bCenterTextX, bool bCenterTextY, const TCHAR* Text, const FFontRenderInfo& RenderInfo) ;
+	int32 WrappedPrint(bool Draw, float X, float Y, int32& out_XL, int32& out_YL, const UFont* Font, float ScaleX, float ScaleY, bool bCenterTextX, bool bCenterTextY, const TCHAR* Text, const FFontRenderInfo& RenderInfo) ;
 	
 	/** Draws a string of text to the screen 
 	 * @param InFont - font to draw with
@@ -280,9 +281,9 @@ public:
 	 *
 	 * returns the Y extent of the rendered text
 	 */
-	float DrawText(UFont* InFont, const FString& InText, float X, float Y, float XScale = 1.f, float YScale = 1.f, const FFontRenderInfo& RenderInfo = FFontRenderInfo());
+	float DrawText(const UFont* InFont, const FString& InText, float X, float Y, float XScale = 1.f, float YScale = 1.f, const FFontRenderInfo& RenderInfo = FFontRenderInfo());
 
-	float DrawText(UFont* InFont, const FText& InText, float X, float Y, float XScale = 1.f, float YScale = 1.f, const FFontRenderInfo& RenderInfo = FFontRenderInfo());
+	float DrawText(const UFont* InFont, const FText& InText, float X, float Y, float XScale = 1.f, float YScale = 1.f, const FFontRenderInfo& RenderInfo = FFontRenderInfo());
 
 	enum ELastCharacterIndexFormat
 	{
@@ -365,7 +366,7 @@ public:
 	 * @param XL out - horizontal length of string
 	 * @param YL out - vertical length of string
 	 */
-	void StrLen(UFont* InFont, const FString& InText, float& XL, float& YL);
+	void StrLen(const UFont* InFont, const FString& InText, float& XL, float& YL);
 
 	/** 
 	 * Calculates the horizontal and vertical size of a given string. This is used for clipped text as it does not take wrapping into account.
@@ -376,7 +377,7 @@ public:
 	 * @param ScaleX - scale that the string is expected to draw at horizontally
 	 * @param ScaleY - scale that the string is expected to draw at vertically
 	 */
-	void TextSize( UFont* InFont, const FString& InText, float& XL, float& YL, float ScaleX=1.f, float ScaleY=1.f);
+	void TextSize( const UFont* InFont, const FString& InText, float& XL, float& YL, float ScaleX=1.f, float ScaleY=1.f);
 	
 	/** Set DrawColor with a FLinearColor and optional opacity override */
 	void SetLinearDrawColor(FLinearColor InColor, float OpacityOverride=-1.f);
diff --git a/Engine/Source/Runtime/Engine/Classes/Engine/Font.h b/Engine/Source/Runtime/Engine/Classes/Engine/Font.h
--- a/Engine/Source/Runtime/Engine/Classes/Engine/Font.h
+++ b/Engine/Source/Runtime/Engine/Classes/Engine/Font.h
@@ -1,9 +1,21 @@
-// Copyright 1998-2014 Epic Games, Inc. All Rights Reserved.
+﻿// Copyright 1998-2014 Epic Games, Inc. All Rights Reserved.
 
 #pragma once
 #include "Engine/FontImportOptions.h"
+#include "FontProviderInterface.h"
 #include "Font.generated.h"
 
+/** What kind of caching is this font using? */
+UENUM()
+enum class EFontCacheType
+{
+	/** The font is using offline caching (this is how UFont traditionally worked) */
+	Offline,
+
+	/** The font is using runtime caching (this is how Slate fonts work) */
+	Runtime,
+};
+
 /** this struct is serialized using native serialization so any changes to it require a package version bump */
 USTRUCT()
 struct FFontCharacter
@@ -69,19 +81,26 @@ struct TStructOpsTypeTraits<FFontCharacter> : public TStructOpsTypeTraitsBase
 };
 
 /**
- * A font object, containing information about a set of glyphs.
- * The glyph bitmaps are stored in the contained textures, while
- * the font database only contains the coordinates of the individual
- * glyph.
+ * A font object, for use by Slate, UMG, and Canvas.
+ *
+ * A font can either be:
+ *   ● Runtime cached - The font contains a series of TTF files that combine to form a composite font. The glyphs are cached on demand when required at runtime.
+ *   ● Offline cached - The font contains a series of textures containing pre-baked cached glyphs and their associated texture coordinates.
  */
 UCLASS(hidecategories=Object, autoexpandcategories=Font, MinimalAPI, BlueprintType)
-class UFont : public UObject
+class UFont : public UObject, public IFontProviderInterface
 {
 	GENERATED_UCLASS_BODY()
 
+	~UFont();
+
+	/** What kind of font caching should we use? This controls which options we see */
+	UPROPERTY(EditAnywhere, Category=Font)
+	TEnumAsByte<EFontCacheType> FontCacheType;
+
 	/** List of characters in the font.  For a MultiFont, this will include all characters in all sub-fonts!  Thus,
 		the number of characters in this array isn't necessary the number of characters available in the font */
-	UPROPERTY(EditAnywhere, Category=Font)
+	UPROPERTY(EditAnywhere, Category=OfflineFont)
 	TArray<struct FFontCharacter> Characters;
 
 	/** Textures that store this font's glyph image data */
@@ -95,27 +114,27 @@ class UFont : public UObject
 	int32 IsRemapped;
 
 	/** Font metrics. */
-	UPROPERTY(EditAnywhere, Category=Font)
+	UPROPERTY(EditAnywhere, Category=OfflineFont)
 	float EmScale;
 
 	/** @todo document */
-	UPROPERTY(EditAnywhere, Category=Font)
+	UPROPERTY(EditAnywhere, Category=OfflineFont)
 	float Ascent;
 
 	/** @todo document */
-	UPROPERTY(EditAnywhere, Category=Font)
+	UPROPERTY(EditAnywhere, Category=OfflineFont)
 	float Descent;
 
 	/** @todo document */
-	UPROPERTY(EditAnywhere, Category=Font)
+	UPROPERTY(EditAnywhere, Category=OfflineFont)
 	float Leading;
 
 	/** Default horizontal spacing between characters when rendering text with this font */
-	UPROPERTY(EditAnywhere, Category=Font)
+	UPROPERTY(EditAnywhere, Category=OfflineFont)
 	int32 Kerning;
 
 	/** Options used when importing this font */
-	UPROPERTY(EditAnywhere, Category=Font)
+	UPROPERTY(EditAnywhere, Category=OfflineFont)
 	struct FFontImportOptionsData ImportOptions;
 
 	/** Number of characters in the font, not including multiple instances of the same character (for multi-fonts).
@@ -130,9 +149,21 @@ class UFont : public UObject
 	TArray<int32> MaxCharHeight;
 
 	/** Scale to apply to the font. */
-	UPROPERTY(EditAnywhere, Category=Font)
+	UPROPERTY(EditAnywhere, Category=OfflineFont)
 	float ScalingFactor;
 
+	/** The default size of the font used for legacy Canvas APIs that don't specify a font size */
+	UPROPERTY(EditAnywhere, Category=RuntimeFont)
+	int32 LegacyFontSize;
+
+	/** The default font name to use for legacy Canvas APIs that don't specify a font name */
+	UPROPERTY(EditAnywhere, Category=RuntimeFont)
+	FName LegacyFontName;
+
+	/** Embedded composite font data */
+	UPROPERTY()
+	FCompositeFont CompositeFont;
+
 
 public:
 	/** This is the character that RemapChar will return if the specified character doesn't exist in the font */
@@ -141,6 +172,18 @@ public:
 	/** When IsRemapped is true, this array maps unicode values to entries in the Characters array */
 	TMap<uint16,uint16> CharRemap;
 
+	/** IFontProviderInterface */
+	virtual const FCompositeFont* GetCompositeFont() const override
+	{
+		return &CompositeFont;
+	}
+
+	/** Get the info needed to use this UFont with Slate, using the fallback data for legacy Canvas APIs */
+	FORCEINLINE FSlateFontInfo GetLegacySlateFontInfo() const
+	{
+		return FSlateFontInfo(this, LegacyFontSize, LegacyFontName);
+	}
+
 	/**
 	 * Returns the size of the object/ resource for display to artists/ LDs in the Editor.
 	 *
@@ -149,82 +192,26 @@ public:
 	virtual SIZE_T GetResourceSize(EResourceSizeMode::Type Mode) override;
 
 	// UFont interface
-	FORCEINLINE TCHAR RemapChar(TCHAR CharCode) const
-	{
-		const uint16 UCode = CharCast<UCS2CHAR>(CharCode);
-		if ( IsRemapped )
-		{
-			// currently, fonts are only remapped if they contain Unicode characters.
-			// For remapped fonts, all characters in the CharRemap map are valid, so
-			// if the characters exists in the map, it's safe to use - otherwise, return
-			// the null character (an empty square on windows)
-			const uint16* FontChar = CharRemap.Find(UCode);
-			if ( FontChar == NULL )
-				return NULLCHARACTER;
-
-			return (TCHAR)*FontChar;
-		}
-
-		// Otherwise, our Characters array will contains 256 members, and is
-		// a one-to-one mapping of character codes to array indexes, though
-		// not every character is a valid character.
-		if ( UCode >= NumCharacters )
-		{
-			return NULLCHARACTER;
-		}
-
-		// If the character's size is 0, it's non-printable or otherwise unsupported by
-		// the font.  Return the default null character (an empty square on windows).
-		if ( !Characters.IsValidIndex(UCode) || (Characters[UCode].VSize == 0 && UCode >= TEXT(' ')) )
-		{
-			return NULLCHARACTER;
-		}
-
-		return CharCode;
-	}
+	ENGINE_API TCHAR RemapChar(TCHAR CharCode) const;
 
-	FORCEINLINE void GetCharSize(TCHAR InCh, float& Width, float& Height) const
-	{
-		Width = Height = 0.f;
-
-		const int32 Ch = (int32)RemapChar(InCh);
-		if( Ch < Characters.Num() )
-		{
-			const FFontCharacter& Char = Characters[Ch];
-			if( Char.TextureIndex < Textures.Num() && Textures[Char.TextureIndex] != NULL )
-			{
-				Width = Char.USize;
-
-				// The height of the character will always be the maximum height of any character in this
-				// font.  This ensures consistent vertical alignment of text.  For example, we don't want
-				// vertically centered text to visually shift up and down as characters are added to a string.
-				// NOTE: This also gives us consistent alignment with fonts generated by the legacy importer.
-				const int32 MultiFontIndex = Ch / NumCharacters;
-				Height = MaxCharHeight[ MultiFontIndex ];
-			}
-		}
-	}
+	/**
+	 * Calculate the width and height of a single character using this font's default size and scale.
+	 *
+	 * @param	InCh					the character to size
+	 * @param	Width					the width of the character (in pixels)
+	 * @param	Height					the height of the character (in pixels)
+	 */
+	ENGINE_API void GetCharSize(TCHAR InCh, float& Width, float& Height) const;
 
 	/**
-	 * Calculate the width of the string using this font's default size and scale.
+	 * Gets the kerning value for a pair of characters
 	 *
-	 * @param	Text					the string to size
+	 * @param FirstChar		The first character in the pair
+	 * @param SecondChar	The second character in the pair
 	 *
-	 * @return	the width (in pixels) of the specified text, or 0 if Text was NULL.
+	 * @return The kerning value
 	 */
-	FORCEINLINE int32 GetStringSize( const TCHAR *Text ) const
-	{
-		float	Width, Height, Total;
-
-		Total = 0.0f;
-		while( *Text )
-		{
-			GetCharSize( *Text++, Width, Height );
-			Total += Width;
-		}
-
-		return( FMath::CeilToInt( Total ) );
-	}
+	ENGINE_API int8 GetCharKerning(TCHAR First, TCHAR Second) const;
 
 	/**
 	 * Calculate the width of the string using this font's default size and scale.
@@ -233,19 +220,16 @@ public:
 	 *
 	 * @return	the width (in pixels) of the specified text, or 0 if Text was NULL.
 	 */
-	FORCEINLINE int32 GetStringHeightSize( const TCHAR *Text ) const
-	{
-		float	Width, Height, Total;
-
-		Total = 0.0f;
-		while( *Text )
-		{
-			GetCharSize( *Text++, Width, Height );
-			Total = FMath::Max( Total, Height );
-		}
+	ENGINE_API int32 GetStringSize( const TCHAR *Text ) const;
 
-		return( FMath::CeilToInt( Total ) );
-	}
+	/**
+	 * Calculate the height of the string using this font's default size and scale.
+	 *
+	 * @param	Text					the string to size
+	 *
+	 * @return	the height (in pixels) of the specified text, or 0 if Text was NULL.
+	 */
+	ENGINE_API int32 GetStringHeightSize( const TCHAR *Text ) const;
 
 	// Begin UObject interface
 	virtual void Serialize( FArchive& Ar ) override;
@@ -256,7 +240,7 @@ public:
 	/**
 	 * Caches the character count and maximum character height for this font (as well as sub-fonts, in the multi-font case)
 	 */
-	virtual void CacheCharacterCountAndMaxCharHeight();
+	ENGINE_API void CacheCharacterCountAndMaxCharHeight();
 
 
 	/**
@@ -264,7 +248,7 @@ public:
 	 *
 	 *	@param	InScalingFactor		The scaling factor to set
 	 */
-	virtual void SetFontScalingFactor(float InScalingFactor)
+	FORCEINLINE void SetFontScalingFactor(float InScalingFactor)
 	{
 		ScalingFactor = InScalingFactor;
 	}
@@ -274,17 +258,15 @@ public:
 	 *
 	 *	@return	float		The scaling factor currently set
 	 */
-	virtual float GetFontScalingFactor()
+	FORCEINLINE float GetFontScalingFactor()
 	{
 		return ScalingFactor;
 	}
 
-	/** Returns the maximum height for any character in this font */
-	virtual float GetMaxCharHeight() const;
+	/** Returns the maximum height for any character in this font using this font's default size and scale. */
+	ENGINE_API float GetMaxCharHeight() const;
 	
 	/** Determines the height and width for the passed in string. */
-	virtual void GetStringHeightAndWidth( const FString& InString, int32& Height, int32& Width ) const;
+	ENGINE_API void GetStringHeightAndWidth( const FString& InString, int32& Height, int32& Width ) const;
+	ENGINE_API void GetStringHeightAndWidth( const TCHAR *Text, int32& Height, int32& Width ) const;
 };
-
-
-
diff --git a/Engine/Source/Runtime/Engine/Private/BatchedElements.cpp b/Engine/Source/Runtime/Engine/Private/BatchedElements.cpp
--- a/Engine/Source/Runtime/Engine/Private/BatchedElements.cpp
+++ b/Engine/Source/Runtime/Engine/Private/BatchedElements.cpp
@@ -309,6 +309,7 @@ static void SetBlendState(FRHICommandList& RHICmdList, ESimpleElementBlendMode B
 	case SE_BLEND_Translucent:
 	case SE_BLEND_TranslucentDistanceField:
 	case SE_BLEND_TranslucentDistanceFieldShadowed:
+	case SE_BLEND_TranslucentAlphaOnly:
 		RHICmdList.SetBlendState(TStaticBlendState<CW_RGB, BO_Add, BF_SourceAlpha, BF_InverseSourceAlpha, BO_Add, BF_Zero, BF_One>::GetRHI());
 		break;
 	case SE_BLEND_Additive:
@@ -344,6 +345,7 @@ static void SetHitTestingBlendState(FRHICommandList& RHICmdList, ESimpleElementB
 	case SE_BLEND_Translucent:
 	case SE_BLEND_TranslucentDistanceField:
 	case SE_BLEND_TranslucentDistanceFieldShadowed:
+	case SE_BLEND_TranslucentAlphaOnly:
 		RHICmdList.SetBlendState(TStaticBlendState<CW_RGBA, BO_Add, BF_One, BF_Zero, BO_Add, BF_One, BF_Zero>::GetRHI());
 		break;
 	case SE_BLEND_Additive:
@@ -362,6 +364,7 @@ FGlobalBoundShaderState FBatchedElements::MaskedBoundShaderState;
 FGlobalBoundShaderState FBatchedElements::DistanceFieldBoundShaderState;
 FGlobalBoundShaderState FBatchedElements::HitTestingBoundShaderState;
 FGlobalBoundShaderState FBatchedElements::ColorChannelMaskShaderState;
+FGlobalBoundShaderState FBatchedElements::AlphaOnlyShaderState;
 
 /** Global alpha ref test value for rendering masked batched elements */
 float GBatchedElementAlphaRefVal = 128.f;
@@ -545,6 +548,17 @@ void FBatchedElements::PrepareShaders(
 				    BlendMode
 				    );
 		    }
+			else if(BlendMode == SE_BLEND_TranslucentAlphaOnly)
+			{
+				SetBlendState(RHICmdList, BlendMode);
+
+				TShaderMapRef<FSimpleElementAlphaOnlyPS> AlphaOnlyPixelShader(GetGlobalShaderMap(FeatureLevel));
+				SetGlobalBoundShaderState(RHICmdList, FeatureLevel, AlphaOnlyShaderState, GSimpleElementVertexDeclaration.VertexDeclarationRHI,
+					*VertexShader, *AlphaOnlyPixelShader);
+
+				AlphaOnlyPixelShader->SetParameters(RHICmdList, Texture);
+				AlphaOnlyPixelShader->SetEditorCompositingParameters(RHICmdList, View, DepthTexture);
+			}
 			else if(BlendMode >= SE_BLEND_RGBA_MASK_START && BlendMode <= SE_BLEND_RGBA_MASK_END)
 			{
 				TShaderMapRef<FSimpleElementColorChannelMaskPS> ColorChannelMaskPixelShader(GetGlobalShaderMap(FeatureLevel));
diff --git a/Engine/Source/Runtime/Engine/Private/Font.cpp b/Engine/Source/Runtime/Engine/Private/Font.cpp
--- a/Engine/Source/Runtime/Engine/Private/Font.cpp
+++ b/Engine/Source/Runtime/Engine/Private/Font.cpp
@@ -7,6 +7,7 @@
 #include "EnginePrivate.h"
 #include "Engine/Font.h"
 #include "Engine/FontImportOptions.h"
+#include "SlateBasics.h"
 
 UFontImportOptions::UFontImportOptions(const FObjectInitializer& ObjectInitializer)
 	: Super(ObjectInitializer)
@@ -17,6 +18,20 @@ UFont::UFont(const FObjectInitializer& ObjectInitializer)
 	: Super(ObjectInitializer)
 {
 	ScalingFactor = 1.0f;
+	LegacyFontSize = 9;
+}
+
+UFont::~UFont()
+{
+	if(FontCacheType == EFontCacheType::Runtime)
+	{
+		TSharedPtr<FSlateRenderer> SlateRenderer = FSlateApplication::Get().GetRenderer();
+		if(SlateRenderer.IsValid())
+		{
+			TSharedRef<FSlateFontCache> FontCache = SlateRenderer->GetFontCache();
+			FontCache->FlushObject(this);
+		}
+	}
 }
 
 void UFont::Serialize( FArchive& Ar )
@@ -29,11 +44,9 @@ void UFont::PostLoad()
 {
 	Super::PostLoad();
 
-
 	// Cache the character count and the maximum character height for this font
 	CacheCharacterCountAndMaxCharHeight();
 
-
 	for( int32 TextureIndex = 0 ; TextureIndex < Textures.Num() ; ++TextureIndex )
 	{
 		UTexture2D* Texture = Textures[TextureIndex];
@@ -53,8 +66,6 @@ void UFont::PostLoad()
 	}
 }
 
-
-
 void UFont::CacheCharacterCountAndMaxCharHeight()
 {
 	// Cache the number of characters in the font.  Obviously this is pretty simple, but note that it will be
@@ -73,34 +84,205 @@ void UFont::CacheCharacterCountAndMaxCharHeight()
 	MaxCharHeight.Add( MaxCharHeightForThisFont );
 }
 
-
 bool UFont::IsLocalizedResource()
 {
 	//@todo: maybe this should be a flag?
 	return true;
 }
 
-float UFont::GetMaxCharHeight() const
+TCHAR UFont::RemapChar(TCHAR CharCode) const
 {
-	// @todo: Provide a version of this function that supports multi-fonts properly.  It should take a
-	//    HeightTest parameter and report the appropriate multi-font MaxCharHeight value back.
-	int32 MaxCharHeightForAllMultiFonts = 1;
-	for( int32 CurMultiFontIndex = 0; CurMultiFontIndex < MaxCharHeight.Num(); ++CurMultiFontIndex )
+	const uint16 UCode = CharCast<UCS2CHAR>(CharCode);
+	if ( IsRemapped )
+	{
+		// currently, fonts are only remapped if they contain Unicode characters.
+		// For remapped fonts, all characters in the CharRemap map are valid, so
+		// if the characters exists in the map, it's safe to use - otherwise, return
+		// the null character (an empty square on windows)
+		const uint16* FontChar = CharRemap.Find(UCode);
+		if ( FontChar == NULL )
+			return NULLCHARACTER;
+
+		return (TCHAR)*FontChar;
+	}
+
+	// Otherwise, our Characters array will contains 256 members, and is
+	// a one-to-one mapping of character codes to array indexes, though
+	// not every character is a valid character.
+	if ( UCode >= NumCharacters )
 	{
-		MaxCharHeightForAllMultiFonts = FMath::Max( MaxCharHeightForAllMultiFonts, MaxCharHeight[ CurMultiFontIndex ] );
+		return NULLCHARACTER;
 	}
-	return MaxCharHeightForAllMultiFonts;
+
+	// If the character's size is 0, it's non-printable or otherwise unsupported by
+	// the font.  Return the default null character (an empty square on windows).
+	if ( !Characters.IsValidIndex(UCode) || (Characters[UCode].VSize == 0 && UCode >= TEXT(' ')) )
+	{
+		return NULLCHARACTER;
+	}
+
+	return CharCode;
 }
 
+void UFont::GetCharSize(TCHAR InCh, float& Width, float& Height) const
+{
+	Width = Height = 0.f;
+
+	switch(FontCacheType)
+	{
+	case EFontCacheType::Offline:
+		{
+			const int32 Ch = (int32)RemapChar(InCh);
+			if( Ch < Characters.Num() )
+			{
+				const FFontCharacter& Char = Characters[Ch];
+				if( Char.TextureIndex < Textures.Num() && Textures[Char.TextureIndex] != NULL )
+				{
+					Width = Char.USize;
+
+					// The height of the character will always be the maximum height of any character in this
+					// font.  This ensures consistent vertical alignment of text.  For example, we don't want
+					// vertically centered text to visually shift up and down as characters are added to a string.
+					// NOTE: This also gives us consistent alignment with fonts generated by the legacy importer.
+					const int32 MultiFontIndex = Ch / NumCharacters;
+					Height = MaxCharHeight[ MultiFontIndex ];
+				}
+			}
+		}
+		break;
+
+	case EFontCacheType::Runtime:
+		{
+			TSharedRef<FSlateFontCache> FontCache = FSlateApplication::Get().GetRenderer()->GetFontCache();
+	
+			const float FontScale = 1.0f;
+			const FSlateFontInfo LegacyFontInfo = GetLegacySlateFontInfo();
+			FCharacterList& CharacterList = FontCache->GetCharacterList(LegacyFontInfo, FontScale);
 
+			const FCharacterEntry& Entry = CharacterList[InCh];
+			Width = Entry.HorizontalOffset + Entry.XAdvance;
+
+			// The height of the character will always be the maximum height of any character in this font
+			Height = CharacterList.GetMaxHeight();
+		}
+		break;
+
+	default:
+		break;
+	}
+}
+
+int8 UFont::GetCharKerning(TCHAR First, TCHAR Second) const
+{
+	switch(FontCacheType)
+	{
+	case EFontCacheType::Offline:
+		{
+			return static_cast<int8>(Kerning);
+		}
+		break;
+
+	case EFontCacheType::Runtime:
+		{
+			TSharedRef<FSlateFontCache> FontCache = FSlateApplication::Get().GetRenderer()->GetFontCache();
+	
+			const float FontScale = 1.0f;
+			const FSlateFontInfo LegacyFontInfo = GetLegacySlateFontInfo();
+			FCharacterList& CharacterList = FontCache->GetCharacterList(LegacyFontInfo, FontScale);
+
+			return CharacterList.GetKerning(First, Second);
+		}
+		break;
+
+	default:
+		break;
+	}
+
+	return 0;
+}
+
+int32 UFont::GetStringSize( const TCHAR *Text ) const
+{
+	int32 Width, Height;
+	GetStringHeightAndWidth( Text, Height, Width );
+	return Width;
+}
+
+int32 UFont::GetStringHeightSize( const TCHAR *Text ) const
+{
+	int32 Width, Height;
+	GetStringHeightAndWidth( Text, Height, Width );
+	return Height;
+}
+
+float UFont::GetMaxCharHeight() const
+{
+	float MaxHeight = 0.0f;
+
+	switch(FontCacheType)
+	{
+	case EFontCacheType::Offline:
+		{
+			// @todo: Provide a version of this function that supports multi-fonts properly.  It should take a
+			//    HeightTest parameter and report the appropriate multi-font MaxCharHeight value back.
+			int32 MaxCharHeightForAllMultiFonts = 1;
+			for( int32 CurMultiFontIndex = 0; CurMultiFontIndex < MaxCharHeight.Num(); ++CurMultiFontIndex )
+			{
+				MaxCharHeightForAllMultiFonts = FMath::Max( MaxCharHeightForAllMultiFonts, MaxCharHeight[ CurMultiFontIndex ] );
+			}
+			MaxHeight = MaxCharHeightForAllMultiFonts;
+		}
+		break;
+
+	case EFontCacheType::Runtime:
+		{
+			TSharedRef<FSlateFontCache> FontCache = FSlateApplication::Get().GetRenderer()->GetFontCache();
+	
+			const float FontScale = 1.0f;
+			const FSlateFontInfo LegacyFontInfo(this, LegacyFontSize);
+			FCharacterList& CharacterList = FontCache->GetCharacterList(LegacyFontInfo, FontScale);
+
+			MaxHeight = CharacterList.GetMaxHeight();
+		}
+		break;
+
+	default:
+		break;
+	}
+
+	return MaxHeight;
+}
 
 void UFont::GetStringHeightAndWidth( const FString& InString, int32& Height, int32& Width ) const
 {
-	Height = GetStringHeightSize( *InString );
-	Width = GetStringSize( *InString );
+	GetStringHeightAndWidth( *InString, Height, Width );
 }
 
+void UFont::GetStringHeightAndWidth( const TCHAR *Text, int32& Height, int32& Width ) const
+{
+	float TotalWidth = 0.0f;
+	float TotalHeight = 0.0f;
+	const TCHAR* PrevChar = nullptr;
+	while( *Text )
+	{
+		float TmpWidth, TmpHeight;
+		GetCharSize( *Text, TmpWidth, TmpHeight );
+
+		int8 CharKerning = 0;
+		if (PrevChar)
+		{
+			CharKerning = GetCharKerning( *PrevChar, *Text );
+		}
+
+		TotalWidth += TmpWidth + CharKerning;
+		TotalHeight += TmpHeight;
 
+		PrevChar = Text++;
+	}
+
+	Width = FMath::CeilToInt( TotalWidth );
+	Height = FMath::CeilToInt( TotalHeight );
+}
 
 SIZE_T UFont::GetResourceSize(EResourceSizeMode::Type Mode)
 {
diff --git a/Engine/Source/Runtime/Engine/Private/SimpleElementShaders.cpp b/Engine/Source/Runtime/Engine/Private/SimpleElementShaders.cpp
--- a/Engine/Source/Runtime/Engine/Private/SimpleElementShaders.cpp
+++ b/Engine/Source/Runtime/Engine/Private/SimpleElementShaders.cpp
@@ -106,6 +106,11 @@ bool FSimpleElementPS::Serialize(FArchive& Ar)
 	return bShaderHasOutdatedParameters;
 }
 
+FSimpleElementAlphaOnlyPS::FSimpleElementAlphaOnlyPS(const ShaderMetaType::CompiledShaderInitializerType& Initializer) :
+	FSimpleElementPS(Initializer)
+{
+}
+
 FSimpleElementGammaBasePS::FSimpleElementGammaBasePS(const ShaderMetaType::CompiledShaderInitializerType& Initializer) :
 	FSimpleElementPS(Initializer)
 {
@@ -292,6 +297,7 @@ bool FSimpleElementColorChannelMaskPS::Serialize(FArchive& Ar)
 
 IMPLEMENT_SHADER_TYPE(,FSimpleElementVS,TEXT("SimpleElementVertexShader"),TEXT("Main"),SF_Vertex);
 IMPLEMENT_SHADER_TYPE(,FSimpleElementPS, TEXT("SimpleElementPixelShader"), TEXT("Main"), SF_Pixel);
+IMPLEMENT_SHADER_TYPE(,FSimpleElementAlphaOnlyPS, TEXT("SimpleElementPixelShader"), TEXT("AlphaOnlyMain"), SF_Pixel);
 IMPLEMENT_SHADER_TYPE(template<>, FSimpleElementGammaPS_SRGB, TEXT("SimpleElementPixelShader"), TEXT("GammaMain"), SF_Pixel);
 IMPLEMENT_SHADER_TYPE(template<>, FSimpleElementGammaPS_Linear, TEXT("SimpleElementPixelShader"), TEXT("GammaMain"), SF_Pixel);
 IMPLEMENT_SHADER_TYPE(template<>, FSimpleElementMaskedGammaPS_SRGB, TEXT("SimpleElementPixelShader"), TEXT("GammaMaskedMain"), SF_Pixel);
diff --git a/Engine/Source/Runtime/Engine/Private/UserInterface/Canvas.cpp b/Engine/Source/Runtime/Engine/Private/UserInterface/Canvas.cpp
--- a/Engine/Source/Runtime/Engine/Private/UserInterface/Canvas.cpp
+++ b/Engine/Source/Runtime/Engine/Private/UserInterface/Canvas.cpp
@@ -633,6 +633,12 @@ void FCanvas::Flush_GameThread(bool bForce)
 		return;
 	}
 
+	// Update the font cache with new text before elements are drawn
+	{
+		TSharedRef<FSlateFontCache> FontCache = FSlateApplication::Get().GetRenderer()->GetFontCache();
+		FontCache->UpdateCache();
+	}
+
 	// FCanvasSortElement compare class
 	struct FCompareFCanvasSortElement
 	{
@@ -894,7 +900,7 @@ int32 FCanvas::DrawShadowedText( float StartX,float StartY,const FText& Text,cla
 }
 
 
-ENGINE_API void StringSize(UFont* Font,int32& XL,int32& YL,const TCHAR* Text)
+ENGINE_API void StringSize(const UFont* Font,int32& XL,int32& YL,const TCHAR* Text)
 {
 	// this functionality has been moved to a static function in UIString
 	FTextSizingParameters Parameters(Font,1.f,1.f);
@@ -911,7 +917,7 @@ ENGINE_API void StringSize(UFont* Font,int32& XL,int32& YL,const TCHAR* Text)
  * @param	DefaultCharHeight	[out] will be set to the height of the typical character
  * @param	pDefaultChar		if specified, pointer to a single character to use for calculating the default character size
  */
-static void GetDefaultCharSize( UFont* DrawFont, float& DefaultCharWidth, float& DefaultCharHeight, const TCHAR* pDefaultChar=NULL )
+static void GetDefaultCharSize( const UFont* DrawFont, float& DefaultCharWidth, float& DefaultCharHeight, const TCHAR* pDefaultChar=NULL )
 {
 	TCHAR DefaultChar = pDefaultChar != NULL ? *pDefaultChar : TEXT('0');
 	DrawFont->GetCharSize(DefaultChar, DefaultCharWidth, DefaultCharHeight);
@@ -941,10 +947,10 @@ void UCanvas::MeasureStringInternal( FTextSizingParameters& Parameters, const TC
 		const float ScaleX = Parameters.Scaling.X;
 		const float ScaleY = Parameters.Scaling.Y;
 
-
-		const float CharIncrement = ( (float)Parameters.DrawFont->Kerning + Parameters.SpacingAdjust.X ) * ScaleX;
+		const float DefaultCharIncrement = Parameters.SpacingAdjust.X * ScaleX;
 		const float DefaultScaledHeight = DefaultCharHeight * ScaleY + Parameters.SpacingAdjust.Y * ScaleY;
 		const TCHAR* pCurrentPos;
+		const TCHAR* pPrevPos = nullptr;
 		for ( pCurrentPos = pText; *pCurrentPos && pCurrentPos < pText + TextLength; ++pCurrentPos )
 		{
 			float CharWidth, CharHeight;
@@ -956,6 +962,13 @@ void UCanvas::MeasureStringInternal( FTextSizingParameters& Parameters, const TC
 			{
 				CharHeight = DefaultCharHeight;
 			}
+
+			float CharSpacing = DefaultCharIncrement;
+			if ( pPrevPos )
+			{
+				CharSpacing += Parameters.DrawFont->GetCharKerning( *pPrevPos, Ch );
+			}
+
 			CharWidth *= ScaleX;
 			CharHeight *= ScaleY;
 
@@ -965,7 +978,7 @@ void UCanvas::MeasureStringInternal( FTextSizingParameters& Parameters, const TC
 				// if we have another character, append the character spacing
 				if ( *pNextPos )
 				{
-					CharWidth += CharIncrement;
+					CharWidth += CharSpacing;
 				}
 			}
 
@@ -996,11 +1009,12 @@ void UCanvas::MeasureStringInternal( FTextSizingParameters& Parameters, const TC
 					}
 				}
 			}
+
+			pPrevPos = pCurrentPos;
 		}
 
 		OutLastCharacterIndex = pCurrentPos - pText;
 	}
-
 }
 
 void UCanvas::CanvasStringSize( FTextSizingParameters& Parameters, const TCHAR* const pText )
@@ -1356,7 +1370,7 @@ void UCanvas::DrawTile( UTexture* Tex, float X, float Y, float XL, float YL, flo
 }
 
 
-void UCanvas::ClippedStrLen( UFont* Font, float ScaleX, float ScaleY, int32& XL, int32& YL, const TCHAR* Text )
+void UCanvas::ClippedStrLen( const UFont* Font, float ScaleX, float ScaleY, int32& XL, int32& YL, const TCHAR* Text )
 {
 	XL = 0;
 	YL = 0;
@@ -1370,7 +1384,7 @@ void UCanvas::ClippedStrLen( UFont* Font, float ScaleX, float ScaleY, int32& XL,
 	}
 }
 
-void VARARGS UCanvas::WrappedStrLenf( UFont* Font, float ScaleX, float ScaleY, int32& XL, int32& YL, const TCHAR* Fmt, ... ) 
+void VARARGS UCanvas::WrappedStrLenf( const UFont* Font, float ScaleX, float ScaleY, int32& XL, int32& YL, const TCHAR* Fmt, ... ) 
 {
 	TCHAR Text[4096];
 	GET_VARARGS( Text, ARRAY_COUNT(Text), ARRAY_COUNT(Text)-1, Fmt, Fmt );
@@ -1379,7 +1393,7 @@ void VARARGS UCanvas::WrappedStrLenf( UFont* Font, float ScaleX, float ScaleY, i
 	WrappedPrint( false, 0.0f, 0.0f, XL, YL, Font, ScaleX, ScaleY, false, false, Text, Info ); 
 }
 
-float UCanvas::DrawText(UFont* InFont, const FText& InText, float X, float Y, float XScale, float YScale, const FFontRenderInfo& RenderInfo)
+float UCanvas::DrawText(const UFont* InFont, const FText& InText, float X, float Y, float XScale, float YScale, const FFontRenderInfo& RenderInfo)
 {
 	ensure(InFont);
 	int32		XL		= 0;
@@ -1399,12 +1413,12 @@ float UCanvas::DrawText(UFont* InFont, const FText& InText, float X, float Y, fl
 	return (float)YL;
 }
 
-float UCanvas::DrawText(UFont* InFont, const FString& InText, float X, float Y, float XScale, float YScale, const FFontRenderInfo& RenderInfo)
+float UCanvas::DrawText(const UFont* InFont, const FString& InText, float X, float Y, float XScale, float YScale, const FFontRenderInfo& RenderInfo)
 {
 	return DrawText(InFont, FText::FromString(InText), X, Y, XScale, YScale, RenderInfo);
 }
 
-int32 UCanvas::WrappedPrint(bool Draw, float X, float Y, int32& out_XL, int32& out_YL, UFont* Font, float ScaleX, float ScaleY, bool bCenterTextX, bool bCenterTextY, const TCHAR* Text, const FFontRenderInfo& RenderInfo) 
+int32 UCanvas::WrappedPrint(bool Draw, float X, float Y, int32& out_XL, int32& out_YL, const UFont* Font, float ScaleX, float ScaleY, bool bCenterTextX, bool bCenterTextY, const TCHAR* Text, const FFontRenderInfo& RenderInfo) 
 {
 	if (ClipX < 0 || ClipY < 0)
 	{
@@ -1476,7 +1490,7 @@ int32 UCanvas::WrappedPrint(bool Draw, float X, float Y, int32& out_XL, int32& o
 	return WrappedStrings.Num();
 }
 
-void UCanvas::StrLen( UFont* InFont, const FString& InText, float& XL, float& YL)
+void UCanvas::StrLen( const UFont* InFont, const FString& InText, float& XL, float& YL)
 {
 	if (InFont == NULL)
 	{
@@ -1492,7 +1506,7 @@ void UCanvas::StrLen( UFont* InFont, const FString& InText, float& XL, float& YL
 	}
 }
 
-void UCanvas::TextSize(UFont* InFont, const FString& InText, float& XL, float& YL, float ScaleX, float ScaleY)
+void UCanvas::TextSize(const UFont* InFont, const FString& InText, float& XL, float& YL, float ScaleX, float ScaleY)
 {
 	int32 XLi, YLi;
 
diff --git a/Engine/Source/Runtime/Engine/Private/UserInterface/CanvasItem.cpp b/Engine/Source/Runtime/Engine/Private/UserInterface/CanvasItem.cpp
--- a/Engine/Source/Runtime/Engine/Private/UserInterface/CanvasItem.cpp
+++ b/Engine/Source/Runtime/Engine/Private/UserInterface/CanvasItem.cpp
@@ -8,7 +8,6 @@
 #include "CanvasItem.h"
 #include "TileRendering.h"
 #include "RHIStaticStates.h"
-#include "EnginePrivate.h"
 #include "Internationalization.h"
 
 
@@ -774,9 +773,6 @@ void FCanvasTextItem::Draw( class FCanvas* InCanvas )
 		return;
 	}
 
-	XScale = Scale.X;
-	YScale = Scale.Y;
-
 	bool bHasShadow = ShadowOffset.Size() != 0.0f;
 	if( ( FontRenderInfo.bEnableShadow == true ) && ( bHasShadow == false ) )
 	{
@@ -800,10 +796,12 @@ void FCanvasTextItem::Draw( class FCanvas* InCanvas )
 		{
 			bHasShadow = false;
 		}
-	}	
-
-	CharIncrement = ( (float)Font->Kerning + HorizSpacingAdjust ) * Scale.X;
-	DrawnSize.Y = Font->GetMaxCharHeight() * Scale.Y;
+	}
+	if (GetFontCacheType() == EFontCacheType::Runtime)
+	{
+		// The runtime font cache uses an alpha-only texture, so we have to force this blend mode so we use the correct shader
+		BlendMode = SE_BLEND_TranslucentAlphaOnly;
+	}
 
 	FVector2D DrawPos( Position.X , Position.Y );
 
@@ -843,7 +841,6 @@ void FCanvasTextItem::Draw( class FCanvas* InCanvas )
 	
 	FLinearColor DrawColor;
 	BatchedElements = NULL;
-	TextLen = Text.ToString().Len();
 	if( bOutlined )
 	{
 		DrawColor = OutlineColor;
@@ -865,21 +862,52 @@ void FCanvasTextItem::Draw( class FCanvas* InCanvas )
 	DrawColor = Color;
 	DrawColor.A *= InCanvas->AlphaModulate;	
 	DrawStringInternal( InCanvas, DrawPos, DrawColor );
-	
-	return;
 }
 
-void FCanvasTextItem::DrawStringInternal( class FCanvas* InCanvas, const FVector2D& DrawPos, const FLinearColor& InColour )
+EFontCacheType FCanvasTextItem::GetFontCacheType() const
+{
+	return Font->FontCacheType;
+}
+
+void FCanvasTextItem::DrawStringInternal( FCanvas* InCanvas, const FVector2D& DrawPos, const FLinearColor& InColor )
 {
-	DrawnSize = FVector2D(EForceInit::ForceInitToZero);
+	switch(GetFontCacheType())
+	{
+	case EFontCacheType::Offline:
+		DrawStringInternal_OfflineCache(InCanvas, DrawPos, InColor);
+		break;
+
+	case EFontCacheType::Runtime:
+		DrawStringInternal_RuntimeCache(InCanvas, DrawPos, InColor);
+		break;
+
+	default:
+		break;
+	}
+}
+
+void FCanvasTextItem::DrawStringInternal_OfflineCache( FCanvas* InCanvas, const FVector2D& DrawPos, const FLinearColor& InColor )
+{
+	DrawnSize = FVector2D::ZeroVector;
+
+	// Nothing to do if no text
+	const FString& TextString = Text.ToString();
+	if( TextString.Len() == 0 )
+	{
+		return;
+	}
+
 	FVector2D CurrentPos = FVector2D(EForceInit::ForceInitToZero);
 	FHitProxyId HitProxyId = InCanvas->GetHitProxyId();
 	FTexture* LastTexture = NULL;
-	UTexture2D* Tex = NULL;
+	UTexture* Tex = NULL;
 	FVector2D InvTextureSize(1.0f,1.0f);
-	const TArray< TCHAR >& Chars = Text.ToString().GetCharArray();
+
+	const float CharIncrement = ( (float)Font->Kerning + HorizSpacingAdjust ) * Scale.X;
+
+	const TArray< TCHAR >& Chars = TextString.GetCharArray();
 	// Draw all characters in string.
-	for( int32 i=0; i < TextLen; i++ )
+	for( int32 i=0; i < TextString.Len(); i++ )
 	{
 		int32 Ch = (int32)Font->RemapChar(Chars[i]);
 
@@ -889,22 +917,22 @@ void FCanvasTextItem::DrawStringInternal( class FCanvas* InCanvas, const FVector
 			continue;
 		}
 
-		FFontCharacter& Char = Font->Characters[Ch];
+		const FFontCharacter& Char = Font->Characters[Ch];
 
 		if (DrawnSize.Y == 0)
 		{
 			// We have a valid character so initialize vertical DrawnSize
-			DrawnSize.Y = Font->GetMaxCharHeight() * YScale;
+			DrawnSize.Y = Font->GetMaxCharHeight() * Scale.Y;
 		}
 
 		if (FChar::IsLinebreak(Chars[i]))
 		{
 			// Set current character offset to the beginning of next line.
 			CurrentPos.X = 0.0f;
-			CurrentPos.Y += Font->GetMaxCharHeight() * YScale;
+			CurrentPos.Y += Font->GetMaxCharHeight() * Scale.Y;
 
 			// Increase the vertical DrawnSize
-			DrawnSize.Y += Font->GetMaxCharHeight() * YScale;
+			DrawnSize.Y += Font->GetMaxCharHeight() * Scale.Y;
 
 			// Don't draw newline character
 			continue;
@@ -924,16 +952,15 @@ void FCanvasTextItem::DrawStringInternal( class FCanvas* InCanvas, const FVector
 				//BatchedElements->AddReserveTriangles(TextLen*2,Tex->Resource,BlendMode);
 				//BatchedElements->AddReserveVertices(TextLen*4);
 
-				FIntPoint ImportedTextureSize = Tex->GetImportedSize();
-				InvTextureSize.X = 1.0f / (float)ImportedTextureSize.X;
-				InvTextureSize.Y = 1.0f / (float)ImportedTextureSize.Y;
+				InvTextureSize.X = 1.0f / Tex->GetSurfaceWidth();
+				InvTextureSize.Y = 1.0f / Tex->GetSurfaceHeight();
 			}
 			LastTexture = Tex->Resource;
 
 			const float X		= CurrentPos.X + DrawPos.X;
-			const float Y		= CurrentPos.Y + DrawPos.Y + Char.VerticalOffset * YScale;
-			float SizeX			= Char.USize * XScale;
-			const float SizeY	= Char.VSize * YScale;
+			const float Y		= CurrentPos.Y + DrawPos.Y + Char.VerticalOffset * Scale.Y;
+			float SizeX			= Char.USize * Scale.X;
+			const float SizeY	= Char.VSize * Scale.Y;
 			const float U		= Char.StartU * InvTextureSize.X;
 			const float V		= Char.StartV * InvTextureSize.Y;
 			const float SizeU	= Char.USize * InvTextureSize.X;
@@ -948,22 +975,22 @@ void FCanvasTextItem::DrawStringInternal( class FCanvas* InCanvas, const FVector
 			int32 V00 = BatchedElements->AddVertex(
 				FVector4( Left, Top, 0.f, Depth ),
 				FVector2D( U, V ),
-				InColour,
+				InColor,
 				HitProxyId );
 			int32 V10 = BatchedElements->AddVertex(
 				FVector4( Right, Top, 0.0f, Depth ),
 				FVector2D( U + SizeU, V ),			
-				InColour,
+				InColor,
 				HitProxyId );
 			int32 V01 = BatchedElements->AddVertex(
 				FVector4( Left, Bottom, 0.0f, Depth ),
 				FVector2D( U, V + SizeV ),	
-				InColour,
+				InColor,
 				HitProxyId);
 			int32 V11 = BatchedElements->AddVertex(
 				FVector4( Right, Bottom, 0.0f, Depth ),
 				FVector2D( U + SizeU, V + SizeV ),
-				InColour,
+				InColor,
 				HitProxyId);
 
 			BatchedElements->AddTriangle(V00, V10, V11, Tex->Resource, BlendMode, FontRenderInfo.GlowInfo);
@@ -986,7 +1013,159 @@ void FCanvasTextItem::DrawStringInternal( class FCanvas* InCanvas, const FVector
 		}
 	}
 }
-/* Draw the item. */
+
+void FCanvasTextItem::DrawStringInternal_RuntimeCache( FCanvas* InCanvas, const FVector2D& DrawPos, const FLinearColor& InColor )
+{
+	DrawnSize = FVector2D::ZeroVector;
+
+	// Nothing to do if no text
+	const FString& TextString = Text.ToString();
+	if( TextString.Len() == 0 )
+	{
+		return;
+	}
+
+	TSharedRef<FSlateFontCache> FontCache = FSlateApplication::Get().GetRenderer()->GetFontCache();
+	
+	const float FontScale = 1.0f;
+	const FSlateFontInfo LegacyFontInfo = (SlateFontInfo.IsSet()) ? SlateFontInfo.GetValue() : Font->GetLegacySlateFontInfo();
+	FCharacterList& CharacterList = FontCache->GetCharacterList( LegacyFontInfo, FontScale );
+
+	const float MaxHeight = CharacterList.GetMaxHeight();
+
+	FHitProxyId HitProxyId = InCanvas->GetHitProxyId();
+
+	uint32 FontTextureIndex = 0;
+	FTextureResource* FontTexture = nullptr;
+
+	float InvTextureSizeX = 0;
+	float InvTextureSizeY = 0;
+
+	float LineX = 0;
+
+	TCHAR PreviousChar = 0;
+
+	int32 Kerning = 0;
+
+	FVector2D TopLeft(0,0);
+
+	const float PosX = TopLeft.X;
+	float PosY = TopLeft.Y;
+
+	const float ScaledHorizSpacingAdjust = HorizSpacingAdjust * Scale.X;
+
+	LineX = PosX;
+	
+	for( int32 CharIndex = 0; CharIndex < TextString.Len(); ++CharIndex )
+	{
+		const TCHAR CurrentChar = TextString[ CharIndex ];
+
+		if (DrawnSize.Y == 0)
+		{
+			// We have a valid character so initialize vertical DrawnSize
+			DrawnSize.Y = MaxHeight;
+		}
+
+		const bool IsNewline = (CurrentChar == '\n');
+
+		if (IsNewline)
+		{
+			// Move down: we are drawing the next line.
+			PosY += MaxHeight;
+			// Carriage return 
+			LineX = PosX;
+			// Increase the vertical DrawnSize
+			DrawnSize.Y += MaxHeight;
+		}
+		else
+		{
+			const FCharacterEntry& Entry = CharacterList[ CurrentChar ];
+
+			if( FontTexture == nullptr || Entry.TextureIndex != FontTextureIndex )
+			{
+				// Font has a new texture for this glyph. Refresh the batch we use and the index we are currently using
+				FontTextureIndex = Entry.TextureIndex;
+				FontTexture = FontCache->GetEngineTextureResource( FontTextureIndex );
+				check(FontTexture);
+
+				FBatchedElementParameters* BatchedElementParameters = nullptr;
+				BatchedElements = InCanvas->GetBatchedElements(FCanvas::ET_Triangle, BatchedElementParameters, FontTexture, BlendMode, FontRenderInfo.GlowInfo);
+				check(BatchedElements);
+
+				InvTextureSizeX = 1.0f/FontTexture->GetSizeX();
+				InvTextureSizeY = 1.0f/FontTexture->GetSizeY();
+			}
+
+			const bool bIsWhitespace = FChar::IsWhitespace(CurrentChar);
+
+			if( !bIsWhitespace && CharIndex > 0 )
+			{
+				Kerning = CharacterList.GetKerning( PreviousChar, CurrentChar );
+			}
+			else
+			{
+				Kerning = 0;
+			}
+
+			LineX += Kerning;
+			PreviousChar = CurrentChar;
+
+			if( !bIsWhitespace )
+			{
+				const float X = LineX + Entry.HorizontalOffset + DrawPos.X;
+				// Note PosX,PosY is the upper left corner of the bounding box representing the string.  This computes the Y position of the baseline where text will sit
+
+				const float Y = PosY - Entry.VerticalOffset + MaxHeight + Entry.GlobalDescender + DrawPos.Y;
+				const float U = Entry.StartU * InvTextureSizeX;
+				const float V = Entry.StartV * InvTextureSizeY;
+				const float SizeX = Entry.USize;
+				const float SizeY = Entry.VSize;
+				const float SizeU = Entry.USize * InvTextureSizeX;
+				const float SizeV = Entry.VSize * InvTextureSizeY;
+
+				float Left, Top, Right, Bottom;
+				Left = X * Depth;
+				Top = Y * Depth;
+				Right = (X + SizeX) * Depth;
+				Bottom = (Y + SizeY) * Depth;
+
+				int32 V00 = BatchedElements->AddVertex(
+					FVector4( Left, Top, 0.f, Depth ),
+					FVector2D( U, V ),
+					InColor,
+					HitProxyId );
+				int32 V10 = BatchedElements->AddVertex(
+					FVector4( Right, Top, 0.0f, Depth ),
+					FVector2D( U + SizeU, V ),			
+					InColor,
+					HitProxyId );
+				int32 V01 = BatchedElements->AddVertex(
+					FVector4( Left, Bottom, 0.0f, Depth ),
+					FVector2D( U, V + SizeV ),	
+					InColor,
+					HitProxyId);
+				int32 V11 = BatchedElements->AddVertex(
+					FVector4( Right, Bottom, 0.0f, Depth ),
+					FVector2D( U + SizeU, V + SizeV ),
+					InColor,
+					HitProxyId);
+
+				BatchedElements->AddTriangle(V00, V10, V11, FontTexture, BlendMode, FontRenderInfo.GlowInfo);
+				BatchedElements->AddTriangle(V00, V11, V01, FontTexture, BlendMode, FontRenderInfo.GlowInfo);
+			}
+
+			LineX += Entry.XAdvance;
+			LineX += ScaledHorizSpacingAdjust;
+
+			// Increase the Horizontal DrawnSize
+			if (LineX > DrawnSize.X)
+			{
+				DrawnSize.X = LineX;
+			}
+		}
+	}
+}
+
 void FCanvasLineItem::Draw( class FCanvas* InCanvas )
 {
 	SCOPE_CYCLE_COUNTER(STAT_Canvas_LineItemTime);
diff --git a/Engine/Source/Runtime/Engine/Public/BatchedElements.h b/Engine/Source/Runtime/Engine/Public/BatchedElements.h
--- a/Engine/Source/Runtime/Engine/Public/BatchedElements.h
+++ b/Engine/Source/Runtime/Engine/Public/BatchedElements.h
@@ -309,6 +309,8 @@ private:
 	static FGlobalBoundShaderState HitTestingBoundShaderState;
 	/** bound shader state for color masked elements */
 	static FGlobalBoundShaderState ColorChannelMaskShaderState;
+	/** bound shader state for alpha only fonts */
+	static FGlobalBoundShaderState AlphaOnlyShaderState;
 
 	/**
 	 * Sets the appropriate vertex and pixel shader.
diff --git a/Engine/Source/Runtime/Engine/Public/CanvasItem.h b/Engine/Source/Runtime/Engine/Public/CanvasItem.h
--- a/Engine/Source/Runtime/Engine/Public/CanvasItem.h
+++ b/Engine/Source/Runtime/Engine/Public/CanvasItem.h
@@ -7,6 +7,7 @@
 #pragma once
 
 #include "Engine/EngineTypes.h"
+#include "Engine/Font.h"
 #include "SceneTypes.h"
 
 class FCanvas;
@@ -295,12 +296,13 @@ public:
 	 * @param	InText			String to draw
 	 * @param	InFont			Font to draw with
 	 */
-	FCanvasTextItem( const FVector2D& InPosition, const FText& InText, class UFont* InFont, const FLinearColor& InColor )
+	FCanvasTextItem( const FVector2D& InPosition, const FText& InText, const UFont* InFont, const FLinearColor& InColor )
 		: FCanvasItem( InPosition )
 		, Text( InText )
 		, Font( InFont )
+		, SlateFontInfo()
 		, HorizSpacingAdjust( 0.0f )
-	  	, ForcedViewportHeight( NULL )
+	  	, ForcedViewportHeight( nullptr )
 		, Depth( 1.0f )
 		, ShadowColor( FLinearColor::Black )
 		, ShadowOffset( FVector2D::ZeroVector )
@@ -316,6 +318,36 @@ public:
 		Scale.Set( 1.0f, 1.0f );
 		BlendMode = SE_BLEND_Translucent;
 	};
+
+	/** 	 
+	 * Text item
+	 *
+	 * @param	InPosition		Draw position
+	 * @param	InText			String to draw
+	 * @param	InFontInfo		Font info to draw with
+	 */
+	FCanvasTextItem( const FVector2D& InPosition, const FText& InText, const FSlateFontInfo& InFontInfo, const FLinearColor& InColor )
+		: FCanvasItem( InPosition )
+		, Text( InText )
+		, Font( Cast<const UFont>(InFontInfo.FontObject) )
+		, SlateFontInfo( InFontInfo )
+		, HorizSpacingAdjust( 0.0f )
+	  	, ForcedViewportHeight( nullptr )
+		, Depth( 1.0f )
+		, ShadowColor( FLinearColor::Black )
+		, ShadowOffset( FVector2D::ZeroVector )
+		, DrawnSize( FVector2D::ZeroVector )
+		, bCentreX( false )
+		, bCentreY( false )
+		, bOutlined( false )
+		, OutlineColor( FLinearColor::Black )
+		, bDontCorrectStereoscopic( true )
+		, TileItem( InPosition, FVector2D::ZeroVector, InColor )
+	{
+		SetColor( InColor );
+		Scale.Set( 1.0f, 1.0f );
+		BlendMode = SE_BLEND_TranslucentAlphaOnly;
+	};
 		
 	/** 
 	 * Set the shadow offset and color. 
@@ -351,7 +383,10 @@ public:
 	FText Text;
 	
 	/* Font to draw text with. */
-	class UFont* Font;
+	const UFont* Font;
+
+	/** Font info to draw the text with. */
+	TOptional<FSlateFontInfo> SlateFontInfo;
 
 	/* Horizontal spacing adjustment. */
 	float HorizSpacingAdjust;
@@ -391,40 +426,28 @@ public:
 
 	/* The scale of the text */
 	FVector2D Scale;
+
 protected:
 	/* Background tile used to fixup 3d text issues. */
 	FCanvasTileItem	TileItem;
 
+	/** Get the type of font cache the UFont is using */
+	EFontCacheType GetFontCacheType() const;
+
 	/** 
 	 * Internal string draw
 	 *
 	 * In a method to make it simpler to do effects like shadow, outline
 	 */
 	void DrawStringInternal( FCanvas* InCanvas, const FVector2D& DrawPos, const FLinearColor& DrawColor );
+	void DrawStringInternal_OfflineCache( FCanvas* InCanvas, const FVector2D& DrawPos, const FLinearColor& DrawColor );
+	void DrawStringInternal_RuntimeCache( FCanvas* InCanvas, const FVector2D& DrawPos, const FLinearColor& DrawColor );
 
 	/** 
-	 * These are used bye the DrawStringInternal function. 
+	 * These are used by the DrawStringInternal function. 
 	 */
-	/* String char length. */
-	int32 TextLen;
-
 	/* Used for batching. */
 	FBatchedElements* BatchedElements;
-	
-	/* Font page index. */
-	int32 PageIndex;
-	
-	/* Font scale (Read from font). */
-	float FontScale;
-	
-	/* Render X scale (based on item scale and font scale). */
-	float XScale;
-	
-	/* Render Y scale (based on item scale and font scale). */
-	float YScale;
-
-	float CharIncrement;
-
 };
 
 /* Line item. */
diff --git a/Engine/Source/Runtime/Engine/Public/CanvasTypes.h b/Engine/Source/Runtime/Engine/Public/CanvasTypes.h
--- a/Engine/Source/Runtime/Engine/Public/CanvasTypes.h
+++ b/Engine/Source/Runtime/Engine/Public/CanvasTypes.h
@@ -844,6 +844,6 @@ private:
 * @param YL - out height
 * @param Text - string of text to be measured
 */
-extern ENGINE_API void StringSize( UFont* Font, int32& XL, int32& YL, const TCHAR* Text );
+extern ENGINE_API void StringSize( const UFont* Font, int32& XL, int32& YL, const TCHAR* Text );
 
 
diff --git a/Engine/Source/Runtime/Engine/Public/SceneTypes.h b/Engine/Source/Runtime/Engine/Public/SceneTypes.h
--- a/Engine/Source/Runtime/Engine/Public/SceneTypes.h
+++ b/Engine/Source/Runtime/Engine/Public/SceneTypes.h
@@ -171,6 +171,8 @@ enum ESimpleElementBlendMode
 	SE_BLEND_AlphaComposite,
 	// Like SE_BLEND_Translucent, but modifies destination alpha
 	SE_BLEND_AlphaBlend,
+	// Like SE_BLEND_Translucent, but reads from an alpha-only texture
+	SE_BLEND_TranslucentAlphaOnly,
 
 	SE_BLEND_RGBA_MASK_START,
 	SE_BLEND_RGBA_MASK_END = SE_BLEND_RGBA_MASK_START + 31, //Using 5bit bit-field for red, green, blue, alpha and desaturation
diff --git a/Engine/Source/Runtime/Engine/Public/SimpleElementShaders.h b/Engine/Source/Runtime/Engine/Public/SimpleElementShaders.h
--- a/Engine/Source/Runtime/Engine/Public/SimpleElementShaders.h
+++ b/Engine/Source/Runtime/Engine/Public/SimpleElementShaders.h
@@ -70,6 +70,18 @@ private:
 	FShaderParameter ScreenToPixel;
 };
 
+/**
+ * Simple pixel shader that just reads from an alpha-only texture
+ */
+class FSimpleElementAlphaOnlyPS : public FSimpleElementPS
+{
+	DECLARE_SHADER_TYPE(FSimpleElementAlphaOnlyPS, Global);
+public:
+
+	FSimpleElementAlphaOnlyPS(const ShaderMetaType::CompiledShaderInitializerType& Initializer);
+	FSimpleElementAlphaOnlyPS() {}
+};
+
 /**
  * A pixel shader for rendering a texture on a simple element.
  */
diff --git a/Engine/Source/Runtime/SlateCore/Private/Fonts/CompositeFont.cpp b/Engine/Source/Runtime/SlateCore/Private/Fonts/CompositeFont.cpp
new file mode 100644
--- /dev/null
+++ b/Engine/Source/Runtime/SlateCore/Private/Fonts/CompositeFont.cpp
@@ -0,0 +1,35 @@
+// Copyright 1998-2014 Epic Games, Inc. All Rights Reserved.
+
+#include "SlateCorePrivatePCH.h"
+#include "CompositeFont.h"
+
+FFontData::FFontData(FString InFontFilename, const EFontHinting InHinting)
+	: FontFilename(MoveTemp(InFontFilename))
+	, FontData()
+	, Hinting(InHinting)
+{
+	if (!FFileHelper::LoadFileToArray(FontData, *FontFilename, FILEREAD_Silent))
+	{
+		UE_LOG(LogSlate, Warning, TEXT("Failed to load font data from '%s'"), *FontFilename);
+	}
+}
+
+bool FFontData::SetFont(FString InFontFilename)
+{
+	TArray<uint8> TempFontData;
+	if (!FFileHelper::LoadFileToArray(TempFontData, *InFontFilename, FILEREAD_Silent))
+	{
+		UE_LOG(LogSlate, Warning, TEXT("Failed to load font data from '%s'"), *InFontFilename);
+		return false;
+	}
+
+	FontFilename = MoveTemp(InFontFilename);
+	FontData = MoveTemp(TempFontData);
+	return true;
+}
+
+void FFontData::SetFont(FString InFontFilename, TArray<uint8> InFontData)
+{
+	FontFilename = MoveTemp(InFontFilename);
+	FontData = MoveTemp(InFontData);
+}
diff --git a/Engine/Source/Runtime/SlateCore/Private/Fonts/FontCache.cpp b/Engine/Source/Runtime/SlateCore/Private/Fonts/FontCache.cpp
--- a/Engine/Source/Runtime/SlateCore/Private/Fonts/FontCache.cpp
+++ b/Engine/Source/Runtime/SlateCore/Private/Fonts/FontCache.cpp
@@ -1,6 +1,7 @@
 // Copyright 1998-2014 Epic Games, Inc. All Rights Reserved.
 
 #include "SlateCorePrivatePCH.h"
+#include "LegacySlateFontInfoCache.h"
 
 
 #ifndef WITH_FREETYPE
@@ -59,6 +60,215 @@ static void FreetypeFree( FT_Memory Memory, void* Block )
 
 #endif // WITH_FREETYPE
 
+/**
+ * Cached data for a given typeface
+ */
+class FCachedTypefaceData
+{
+public:
+	/** Default constructor */
+	FCachedTypefaceData()
+		: Typeface(nullptr)
+		, SingularFontData(nullptr)
+		, NameToFontDataMap()
+		, ScalingFactor(1.0f)
+	{
+	}
+
+	/** Construct the cache from the given typeface */
+	FCachedTypefaceData(const FTypeface& InTypeface, const float InScalingFactor = 1.0f)
+		: Typeface(&InTypeface)
+		, SingularFontData(nullptr)
+		, NameToFontDataMap()
+		, ScalingFactor(InScalingFactor)
+	{
+		if(InTypeface.Fonts.Num() == 0)
+		{
+			// We have no entries - don't bother building a map
+			SingularFontData = nullptr;
+		}
+		else if(InTypeface.Fonts.Num() == 1)
+		{
+			// We have a single entry - don't bother building a map
+			SingularFontData = &InTypeface.Fonts[0].Font;
+		}
+		else
+		{
+			// Add all the entries from the typeface
+			for(const FTypefaceEntry& TypefaceEntry : InTypeface.Fonts)
+			{
+				NameToFontDataMap.Add(TypefaceEntry.Name, &TypefaceEntry.Font);
+			}
+
+			// Add a special "None" entry to return the first font from the typeface
+			if(!NameToFontDataMap.Contains(NAME_None))
+			{
+				NameToFontDataMap.Add(NAME_None, &InTypeface.Fonts[0].Font);
+			}
+		}
+	}
+
+	/** Get the typeface we cached data from */
+	const FTypeface& GetTypeface() const
+	{
+		check(Typeface);
+		return *Typeface;
+	}
+
+	/** Find the font associated with the given name */
+	const FFontData* GetFontData(const FName& InName) const
+	{
+		if(NameToFontDataMap.Num() > 0)
+		{
+			const FFontData* const * const FoundFontData = NameToFontDataMap.Find(InName);
+			return (FoundFontData) ? *FoundFontData : nullptr;
+		}
+		return SingularFontData;
+	}
+
+	/** Get the scaling factor for this typeface */
+	float GetScalingFactor() const
+	{
+		return ScalingFactor;
+	}
+
+private:
+	/** Typeface we cached data from */
+	const FTypeface* Typeface;
+
+	/** Singular entry, used when we don't have enough data to warrant using a map */
+	const FFontData* SingularFontData;
+
+	/** Mapping between a font name, and its data */
+	TMap<FName, const FFontData*> NameToFontDataMap;
+
+	/** Scaling factor to apply to this typeface */
+	float ScalingFactor;
+};
+
+/**
+ * Cached data for a given composite font
+ */
+class FCachedCompositeFontData
+{
+public:
+	/** Default constructor */
+	FCachedCompositeFontData()
+		: CompositeFont(nullptr)
+		, CachedTypefaces()
+		, CachedFontRanges()
+	{
+	}
+
+	/** Construct the cache from the given composite font */
+	FCachedCompositeFontData(const FCompositeFont& InCompositeFont)
+		: CompositeFont(&InCompositeFont)
+		, CachedTypefaces()
+		, CachedFontRanges()
+	{
+		// Add all the entries from the composite font
+		CachedTypefaces.Add(MakeShareable(new FCachedTypefaceData(InCompositeFont.DefaultTypeface)));
+		for(const FCompositeSubFont& SubTypeface : InCompositeFont.SubTypefaces)
+		{
+			TSharedPtr<FCachedTypefaceData> CachedTypeface = MakeShareable(new FCachedTypefaceData(SubTypeface.Typeface, SubTypeface.ScalingFactor));
+			CachedTypefaces.Add(CachedTypeface);
+
+			for(const FInt32Range& Range : SubTypeface.CharacterRanges)
+			{
+				CachedFontRanges.Add(FCachedFontRange(Range, CachedTypeface));
+			}
+		}
+
+		// Sort the font ranges into ascending order
+		CachedFontRanges.Sort([](const FCachedFontRange& RangeOne, const FCachedFontRange& RangeTwo) -> bool
+		{
+			if(RangeOne.Range.IsEmpty() && !RangeTwo.Range.IsEmpty())
+			{
+				return true;
+			}
+			if(!RangeOne.Range.IsEmpty() && RangeTwo.Range.IsEmpty())
+			{
+				return false;
+			}
+			return RangeOne.Range.GetLowerBoundValue() < RangeTwo.Range.GetLowerBoundValue();
+		});
+	}
+
+	/** Get the composite font we cached data from */
+	const FCompositeFont& GetCompositeFont() const
+	{
+		check(CompositeFont);
+		return *CompositeFont;
+	}
+
+	/** Get the default typeface for this composite font */
+	TSharedPtr<const FCachedTypefaceData> GetDefaultTypeface() const
+	{
+		return CachedTypefaces[0];
+	}
+
+	/** Get the typeface that should be used for the given character */
+	TSharedPtr<const FCachedTypefaceData> GetTypefaceForCharacter(const TCHAR InChar) const
+	{
+		const int32 CharIndex = static_cast<int32>(InChar);
+
+		for(const FCachedFontRange& CachedRange : CachedFontRanges)
+		{
+			if(CachedRange.Range.IsEmpty())
+			{
+				continue;
+			}
+
+			// Ranges are sorting in ascending order (by the start position), so if this range starts higher than the character we're looking for, we can bail from the check
+			if(CachedRange.Range.GetLowerBoundValue() > CharIndex)
+			{
+				break;
+			}
+
+			if(CachedRange.Range.Contains(CharIndex))
+			{
+				return CachedRange.CachedTypeface;
+			}
+		}
+
+		return CachedTypefaces[0];
+	}
+
+private:
+	/** Entry containing a range and the typeface associated with that range */
+	struct FCachedFontRange
+	{
+		/** Default constructor */
+		FCachedFontRange()
+			: Range(FInt32Range::Empty())
+			, CachedTypeface()
+		{
+		}
+
+		/** Construct from the given range and typeface */
+		FCachedFontRange(const FInt32Range& InRange, const TSharedPtr<FCachedTypefaceData>& InCachedTypeface)
+			: Range(InRange)
+			, CachedTypeface(InCachedTypeface)
+		{
+		}
+
+		/** Range to use for the typeface */
+		FInt32Range Range;
+
+		/** Typeface to which the range applies */
+		TSharedPtr<FCachedTypefaceData> CachedTypeface;
+	};
+
+	/** Composite font we cached data from */
+	const FCompositeFont* CompositeFont;
+
+	/** Array of cached typefaces - 0 is the default typeface, and the remaining entries are sub-typefaces */
+	TArray<TSharedPtr<FCachedTypefaceData>> CachedTypefaces;
+
+	/** Array of font ranges paired with their associated typefaces - this is sorted in ascending order */
+	TArray<FCachedFontRange> CachedFontRanges;
+};
+
 /**
  * An interface to the freetype API.                     
  */
@@ -105,11 +315,13 @@ public:
 #if WITH_FREETYPE
 		FontToKerningPairMap.Empty();
 		// toss memory
-		for (TMap<FName, FFontFaceAndMemory>::TIterator It(FontFaceMap); It; ++It)
+		for (auto& FontFaceEntry : FontFaceMap)
 		{
-			FMemory::Free(It.Value().Memory);
+			FT_Done_Face(FontFaceEntry.Value.Face);
+			FMemory::Free(FontFaceEntry.Value.Memory);
 		}
 		FontFaceMap.Empty();
+		CompositeFontToCachedDataMap.Empty();
 #endif // WITH_FREETYPE
 	}
 
@@ -120,12 +332,18 @@ public:
 	 * @param Char			The character to render
 	 * @param OutCharInfo	Will contain the created render data
 	 */
-	void GetRenderData( const FSlateFontInfo& InFontInfo, TCHAR Char, FCharacterRenderData& OutRenderData, float Scale )
+	void GetRenderData( const FSlateFontInfo& InFontInfo, TCHAR Char, FCharacterRenderData& OutRenderData, const float InScale )
 	{
 #if WITH_FREETYPE
+		float SubFontScalingFactor = 1.0f;
+		const FFontData& FontData = GetFontDataForCharacter(InFontInfo, Char, SubFontScalingFactor);
+
+		// Apply the sub-font scale
+		const float FinalScale = InScale * SubFontScalingFactor;
+
 		// Find or load the face if needed
 		FT_UInt GlyphIndex = 0;
-		FT_Face FontFace = GetFontFace( InFontInfo.FontName );
+		FT_Face FontFace = GetFontFace( FontData );
 
 		if ( FontFace != nullptr ) 
 		{
@@ -135,7 +353,7 @@ public:
 
 		uint32 LocalGlyphFlags = GlyphFlags;
 
-		switch(InFontInfo.Hinting)
+		switch(FontData.Hinting)
 		{
 		case EFontHinting::Auto:		LocalGlyphFlags |= FT_LOAD_FORCE_AUTOHINT; break;
 		case EFontHinting::AutoLight:	LocalGlyphFlags |= FT_LOAD_TARGET_LIGHT; break;
@@ -148,58 +366,33 @@ public:
 		// If the requested glyph doesn't exist, use the localization fallback font.
 		if ( FontFace == nullptr || (Char != 0 && GlyphIndex == 0) )
 		{
-			static FName FallbackFontName( NAME_None );
-			if( FallbackFontName == NAME_None )
-			{
-				FText FallbackFontFilename;
-				FallbackFontName = FName( *( FPaths::EngineContentDir() / TEXT("Slate/Fonts/") / ( NSLOCTEXT("Slate", "FallbackFont", "DroidSansFallback").ToString() + TEXT(".ttf") ) ) );
-			}
-			if( FallbackFontName.IsValid() && FallbackFontName != NAME_None )
-			{
-				FontFace = GetFontFace( FallbackFontName );
-				if (FontFace != nullptr)
-				{					
-					GlyphIndex = FT_Get_Char_Index( FontFace, Char );
-					LocalGlyphFlags |= FT_LOAD_FORCE_AUTOHINT;
-				}
-				else
-				{
-					FallbackFontName = FName(NAME_Inactive);
-				}
+			FontFace = GetFontFace( FLegacySlateFontInfoCache::Get().GetFallbackFont() );
+			if (FontFace != nullptr)
+			{					
+				GlyphIndex = FT_Get_Char_Index( FontFace, Char );
+				LocalGlyphFlags |= FT_LOAD_FORCE_AUTOHINT;
 			}
 		}
 
 		// If the requested glyph doesn't exist, use the last resort fallback font.
 		if ( FontFace == nullptr || ( Char != 0 && GlyphIndex == 0 ) )
 		{
-			static FName LastResortFontName( NAME_None );
-			if( LastResortFontName == NAME_None )
-			{
-				LastResortFontName = FName( *( FPaths::EngineContentDir() / TEXT("Slate/Fonts/LastResort.ttf") ) );
-			}
-			if( LastResortFontName.IsValid() && LastResortFontName != NAME_None )
-			{
-				FontFace = GetFontFace( LastResortFontName );
-				check( FontFace );
-				GlyphIndex = FT_Get_Char_Index( FontFace, Char );
-				LocalGlyphFlags |= FT_LOAD_FORCE_AUTOHINT;
-			}
-			else
-			{
-				LastResortFontName = FName(NAME_Inactive);
-			}
+			FontFace = GetFontFace( FLegacySlateFontInfoCache::Get().GetLastResortFont() );
+			check( FontFace );
+			GlyphIndex = FT_Get_Char_Index( FontFace, Char );
+			LocalGlyphFlags |= FT_LOAD_FORCE_AUTOHINT;
 		}
 
 		// Set the character size to render at (needs to be in 1/64 of a "point")
 		FT_Error Error = FT_Set_Char_Size( FontFace, 0, InFontInfo.Size*64, FontCacheConstants::HorizontalDPI, FontCacheConstants::VerticalDPI );
 		check(Error==0);
 
-		if( Scale != 1.0f )
+		if( FinalScale != 1.0f )
 		{
 			FT_Matrix ScaleMatrix;
 			ScaleMatrix.xy = 0;
-			ScaleMatrix.xx = (FT_Fixed)(Scale * 65536);
-			ScaleMatrix.yy = (FT_Fixed)(Scale * 65536);
+			ScaleMatrix.xx = (FT_Fixed)(FinalScale * 65536);
+			ScaleMatrix.yy = (FT_Fixed)(FinalScale * 65536);
 			ScaleMatrix.yx = 0;
 			FT_Set_Transform( FontFace, &ScaleMatrix, nullptr );
 		}
@@ -208,7 +401,6 @@ public:
 			FT_Set_Transform( FontFace, nullptr, nullptr );
 		}
 
-
 		// Load the glyph.  Force using the freetype hinter because not all true type fonts have their own hinting
 		Error = FT_Load_Glyph( FontFace, GlyphIndex, LocalGlyphFlags );
 		check(Error==0);
@@ -221,7 +413,6 @@ public:
 		// one byte per pixel 
 		const uint32 GlyphPixelSize = 1;
 
-	
 		FT_Bitmap* Bitmap = nullptr;
 
 		if( Slot->bitmap.pixel_mode == FT_PIXEL_MODE_MONO )
@@ -242,7 +433,6 @@ public:
 		OutRenderData.RawPixels.Reset();
 		OutRenderData.RawPixels.AddUninitialized( Bitmap->rows * Bitmap->width );
 
-
 		{
 			// Copy the rendered bitmap to our raw pixels array
 
@@ -274,7 +464,7 @@ public:
 		FT_Get_Glyph( Slot, &Glyph );
 		FT_Glyph_Get_CBox( Glyph, FT_GLYPH_BBOX_PIXELS, &GlyphBox );
 
-		int32 Height = (FT_MulFix( FontFace->height, FontFace->size->metrics.y_scale ) / 64) * Scale;
+		int32 Height = (FT_MulFix( FontFace->height, FontFace->size->metrics.y_scale ) / 64) * FinalScale;
 
 		// Set measurement info for this character
 		OutRenderData.Char = Char;
@@ -284,9 +474,9 @@ public:
 
 		// Need to divide by 64 to get pixels;
 		// Ascender is not scaled by freetype.  Scale it now. 
-		OutRenderData.MeasureInfo.GlobalAscender = ( FontFace->size->metrics.ascender / 64 ) * Scale;
+		OutRenderData.MeasureInfo.GlobalAscender = ( FontFace->size->metrics.ascender / 64 ) * InScale; // Don't add the sub-font scale, as we want a consistent ascender
 		// Descender is not scaled by freetype.  Scale it now. 
-		OutRenderData.MeasureInfo.GlobalDescender = ( FontFace->size->metrics.descender / 64 ) * Scale;
+		OutRenderData.MeasureInfo.GlobalDescender = ( FontFace->size->metrics.descender / 64 ) * InScale; // Don't add the sub-font scale, as we want a consistent descender
 		// Note we use Slot->advance instead of Slot->metrics.horiAdvance because Slot->Advance contains transformed position (needed if we scale)
 		OutRenderData.MeasureInfo.XAdvance =  Slot->advance.x / 64;
 		OutRenderData.MeasureInfo.HorizontalOffset = Slot->bitmap_left;
@@ -305,10 +495,10 @@ public:
 	/**
 	 * @param Whether or not the font has kerning
 	 */
-	bool HasKerning( const FSlateFontInfo& InFontInfo )
+	bool HasKerning( const FFontData& InFontData )
 	{
 #if WITH_FREETYPE
-		FT_Face FontFace = GetFontFace( InFontInfo.FontName );
+		FT_Face FontFace = GetFontFace( InFontData );
 
 		if ( FontFace == nullptr )
 		{
@@ -329,25 +519,34 @@ public:
 	 * @param InFontInfo	Information about the font that used to draw the string with the first and second characters
 	 * @return The kerning amount, 0 if no kerning
 	 */
-	int8 GetKerning( TCHAR First, TCHAR Second, const FSlateFontKey& FontKey )
+	int8 GetKerning( TCHAR First, TCHAR Second, const FSlateFontInfo& InFontInfo, float InScale )
 	{
 #if WITH_FREETYPE
+		float SubFontScalingFactor = 1.0f; // this is fine to reuse below, as we only apply kerning for characters using the same font
+		const FFontData& FirstFontData = GetFontDataForCharacter(InFontInfo, First, SubFontScalingFactor);
+		const FFontData& SecondFontData = GetFontDataForCharacter(InFontInfo, Second, SubFontScalingFactor);
+
+		// Apply the sub-font scale
+		const float FinalScale = InScale * SubFontScalingFactor;
+
 		int32 Kerning = 0;
 		int32* FoundKerning = nullptr;
 
-		FT_Face FontFace = GetFontFace( FontKey.FontInfo.FontName );
+		FT_Face FontFace = GetFontFace( FirstFontData );
+		FT_Face SecondFontFace = GetFontFace( SecondFontData );
 
-		// Check if this font has kerning.  Not all fonts do.
-		if( FontFace != nullptr && FT_HAS_KERNING( FontFace ) )
+		// Check if this font has kerning as not all fonts do.
+		// We also can't perform kerning between two separate font faces
+		if( FontFace != nullptr && FontFace == SecondFontFace && FT_HAS_KERNING( FontFace ) )
 		{
-			FT_Error Error = FT_Set_Char_Size( FontFace, 0, FontKey.FontInfo.Size*64, FontCacheConstants::HorizontalDPI, FontCacheConstants::VerticalDPI  );
+			FT_Error Error = FT_Set_Char_Size( FontFace, 0, InFontInfo.Size*64, FontCacheConstants::HorizontalDPI, FontCacheConstants::VerticalDPI  );
 
-			if( FontKey.Scale != 1.0f )
+			if( FinalScale != 1.0f )
 			{
 				FT_Matrix ScaleMatrix;
 				ScaleMatrix.xy = 0;
-				ScaleMatrix.xx = (FT_Fixed)(FontKey.Scale * 65536);
-				ScaleMatrix.yy = (FT_Fixed)(FontKey.Scale * 65536);
+				ScaleMatrix.xx = (FT_Fixed)(FinalScale * 65536);
+				ScaleMatrix.yy = (FT_Fixed)(FinalScale * 65536);
 				ScaleMatrix.yx = 0;
 				FT_Set_Transform( FontFace, &ScaleMatrix, nullptr );
 			}
@@ -381,78 +580,202 @@ public:
 #endif // WITH_FREETYPE
 	}
 
+	/** Get the attributes associated with the given font data */
+	const TSet<FName>& GetFontAttributes( const FFontData& InFontData )
+	{
+		static const TSet<FName> DummyAttributes;
+
+		FFontFaceAndMemory* FaceAndMemory = FontFaceMap.Find(&InFontData);
+		if (!FaceAndMemory)
+		{
+			GetFontFace(InFontData); // will try and create the entry
+			FaceAndMemory = FontFaceMap.Find(&InFontData);
+		}
+
+		return (FaceAndMemory) ? FaceAndMemory->Attributes : DummyAttributes;
+	}
+
 private:
 
-#if WITH_FREETYPE
-	/**
-	 * Gets or loads a freetype font face
-	 *
-	 * @param The name of the font to load
-	 */
-	FT_Face GetFontFace( const FName& FontName )
+	/** Get the cached composite font data for the given composite font */
+	TSharedPtr<const FCachedCompositeFontData> GetCachedCompositeFont(const FCompositeFont* const InCompositeFont)
 	{
-		static const FName SpecialName_DefaultSystemFont("DefaultSystemFont");
+		if(!InCompositeFont)
+		{
+			return nullptr;
+		}
 
-		FFontFaceAndMemory* FaceAndMemory = FontFaceMap.Find( FontName );
-		if (!FaceAndMemory)
+		TSharedPtr<FCachedCompositeFontData> FoundCompositeFontData = CompositeFontToCachedDataMap.FindRef(InCompositeFont);
+		if(FoundCompositeFontData.IsValid())
 		{
-			// make a new entry
-			FaceAndMemory = &FontFaceMap.Add(FontName, FFontFaceAndMemory());
+			return FoundCompositeFontData;
+		}
+
+		return CompositeFontToCachedDataMap.Add(InCompositeFont, MakeShareable(new FCachedCompositeFontData(*InCompositeFont)));
+	}
+
+	/** Get the default typeface for the given composite font */
+	TSharedPtr<const FCachedTypefaceData> GetDefaultCachedTypeface(const FCompositeFont* const InCompositeFont)
+	{
+		TSharedPtr<const FCachedCompositeFontData> CachedCompositeFont = GetCachedCompositeFont(InCompositeFont);
+		return (CachedCompositeFont.IsValid()) ? CachedCompositeFont->GetDefaultTypeface() : nullptr;
+	}
 
-			// default to error condition
-			bool Error = true;
+	/** Get the typeface that should be used for the given character */
+	TSharedPtr<const FCachedTypefaceData> GetCachedTypefaceForCharacter(const FCompositeFont* const InCompositeFont, const TCHAR InChar)
+	{
+		TSharedPtr<const FCachedCompositeFontData> CachedCompositeFont = GetCachedCompositeFont(InCompositeFont);
+		return (CachedCompositeFont.IsValid()) ? CachedCompositeFont->GetTypefaceForCharacter(InChar) : nullptr;
+	}
+	
+	/** Get the default font data to use for the given font info */
+	const FFontData& GetDefaultFontData(const FSlateFontInfo& InFontInfo)
+	{
+		static const FFontData DummyFontData;
 
-			if (FontName == SpecialName_DefaultSystemFont)
+		const FCompositeFont* const ResolvedCompositeFont = InFontInfo.GetCompositeFont();
+		TSharedPtr<const FCachedTypefaceData> CachedTypefaceData = GetDefaultCachedTypeface(ResolvedCompositeFont);
+		if(CachedTypefaceData.IsValid())
+		{
+			// Try to find the correct font from the typeface
+			const FFontData* FoundFontData = CachedTypefaceData->GetFontData(InFontInfo.TypefaceFontName);
+			if(FoundFontData)
 			{
-				// Ask the platform to load the data for a default system font.
-				const TArray<uint8> FontBytes = FPlatformMisc::GetSystemFontBytes();
-				if ( FontBytes.Num() > 0 )
-				{
-					FaceAndMemory->Memory = static_cast<uint8*>( FMemory::Malloc(FontBytes.Num()) );
-					FMemory::Memcpy( FaceAndMemory->Memory, FontBytes.GetData(), FontBytes.Num() );
-					
-					// initialize the font, setting the error code
-					Error = FT_New_Memory_Face( FTLibrary, FaceAndMemory->Memory, static_cast<FT_Long>(FontBytes.Num()), 0, &FaceAndMemory->Face ) != 0;
-				}
+				return *FoundFontData;
 			}
-			else
+
+			// Failing that, return the first font available (the "None" font)
+			FoundFontData = CachedTypefaceData->GetFontData(NAME_None);
+			if(FoundFontData)
+			{
+				return *FoundFontData;
+			}
+		}
+
+		return DummyFontData;
+	}
+
+	/** Get the font data to use for the given font info and character */
+	const FFontData& GetFontDataForCharacter(const FSlateFontInfo& InFontInfo, const TCHAR InChar, float& OutScalingFactor)
+	{
+		static const FFontData DummyFontData;
+
+		const FCompositeFont* const ResolvedCompositeFont = InFontInfo.GetCompositeFont();
+		TSharedPtr<const FCachedTypefaceData> CachedTypefaceData = GetCachedTypefaceForCharacter(ResolvedCompositeFont, InChar);
+		if(CachedTypefaceData.IsValid())
+		{
+			OutScalingFactor = CachedTypefaceData->GetScalingFactor();
+
+			// Try to find the correct font from the typeface
+			const FFontData* FoundFontData = CachedTypefaceData->GetFontData(InFontInfo.TypefaceFontName);
+			if(FoundFontData)
 			{
-				// load the font via UE4 methods, so that it will route through all proper file management (network file loading, etc)
-				const FString FontPath = FontName.ToString();
-				int64 FileSize = IFileManager::Get().FileSize(*FontPath);
-				if ( FileSize > 0 )
+				return *FoundFontData;
+			}
+
+			// Failing that, try and find a font by the attributes of the default font with the given name
+			TSharedPtr<const FCachedTypefaceData> CachedDefaultTypefaceData = GetDefaultCachedTypeface(ResolvedCompositeFont);
+			if(CachedDefaultTypefaceData.IsValid() && CachedTypefaceData != CachedDefaultTypefaceData)
+			{
+				const FFontData* const FoundDefaultFontData = CachedDefaultTypefaceData->GetFontData(InFontInfo.TypefaceFontName);
+				if(FoundDefaultFontData)
 				{
-					// allocate space for the font
-					FaceAndMemory->Memory = (uint8*)FMemory::Malloc((uint32)FileSize);
-					FArchive* Ar = IFileManager::Get().CreateFileReader(*FontPath);
-					if (Ar != nullptr)
+					const TSet<FName>& DefaultFontAttributes = GetFontAttributes(*FoundDefaultFontData);
+					FoundFontData = GetBestMatchFontForAttributes(CachedTypefaceData, DefaultFontAttributes);
+					if(FoundFontData)
 					{
-						// read in the file
-						Ar->Serialize( FaceAndMemory->Memory, FileSize );
-						delete Ar;
-
-						// initialize the font, setting the error code
-						Error = FT_New_Memory_Face( FTLibrary, FaceAndMemory->Memory, (FT_Long)FileSize, 0, &FaceAndMemory->Face ) != 0;
+						return *FoundFontData;
 					}
 				}
 			}
 
+			// Failing that, return the first font available (the "None" font)
+			FoundFontData = CachedTypefaceData->GetFontData(NAME_None);
+			if(FoundFontData)
+			{
+				return *FoundFontData;
+			}
+		}
+
+		OutScalingFactor = 1.0f;
+		return DummyFontData;
+	}
+
+	const FFontData* GetBestMatchFontForAttributes(const TSharedPtr<const FCachedTypefaceData>& InCachedTypefaceData, const TSet<FName>& InFontAttributes)
+	{
+		const FFontData* BestMatchFont = nullptr;
+		int32 BestMatchCount = 0;
+
+		const FTypeface& Typeface = InCachedTypefaceData->GetTypeface();
+		for(const FTypefaceEntry& TypefaceEntry : Typeface.Fonts)
+		{
+			const TSet<FName>& FontAttributes = GetFontAttributes(TypefaceEntry.Font);
+
+			int32 MatchCount = 0;
+			for(const FName& InAttribute : InFontAttributes)
+			{
+				if(FontAttributes.Contains(InAttribute))
+				{
+					++MatchCount;
+				}
+			}
+
+			if(MatchCount > BestMatchCount || !BestMatchFont)
+			{
+				BestMatchFont = &TypefaceEntry.Font;
+				BestMatchCount = MatchCount;
+			}
+		}
+
+		return BestMatchFont;
+	}
+
+#if WITH_FREETYPE
+	/**
+	 * Gets or loads a freetype font face
+	 *
+	 * @param InFontData Information about the font to load
+	 */
+	FT_Face GetFontFace( const FFontData& InFontData )
+	{
+		FFontFaceAndMemory* FaceAndMemory = FontFaceMap.Find(&InFontData);
+		if (!FaceAndMemory && InFontData.FontData.Num() > 0)
+		{
+			// make a new entry
+			FaceAndMemory = &FontFaceMap.Add(&InFontData, FFontFaceAndMemory());
+
+			// todo: jdale - can we avoid this copy without potentially crashing if the source FFontData is unloaded (when using a UObject)?
+			//				 We'd just need to remove this entry from the font cache before unloading the UFont and then we 
+			//				 could just take a pointer to the array data (which we don't own, so won't delete)
+			FaceAndMemory->Memory = static_cast<uint8*>(FMemory::Malloc(InFontData.FontData.Num()));
+			FMemory::Memcpy(FaceAndMemory->Memory, InFontData.FontData.GetData(), InFontData.FontData.Num());
+
+			// initialize the font, setting the error code
+			const bool bFailedToLoadFace = FT_New_Memory_Face(FTLibrary, FaceAndMemory->Memory, static_cast<FT_Long>(InFontData.FontData.Num()), 0, &FaceAndMemory->Face) != 0;
+
 			// if it failed, we don't want to keep the memory around
-			if ( Error )
+			if (bFailedToLoadFace)
 			{
 				FMemory::Free(FaceAndMemory->Memory);
-				FontFaceMap.Remove(FontName);
+				FontFaceMap.Remove(&InFontData);
 				FaceAndMemory = nullptr;
-				UE_LOG( LogSlate, Warning, TEXT("GetFontFace failed to load or process '%s'"), *FontName.ToString());
+				UE_LOG(LogSlate, Warning, TEXT("GetFontFace failed to load or process '%s'"), *InFontData.FontFilename);
 			}
-		}
+			
+			if (FaceAndMemory)
+			{
+				// Parse out the font attributes
+				TArray<FString> Styles;
+				FString(FaceAndMemory->Face->style_name).ParseIntoArray(&Styles, TEXT(" "), true);
 
-		if ( FaceAndMemory == nullptr )
-		{
-			return nullptr;
+				for (const FString& Style : Styles)
+				{
+					FaceAndMemory->Attributes.Add(*Style);
+				}
+			}
 		}
 
-		return FaceAndMemory->Face;
+		return (FaceAndMemory) ? FaceAndMemory->Face : nullptr;
 	}
 
 private:
@@ -464,17 +787,23 @@ private:
 		// the memory for the face (can't be a TArray as the FontFaceMap could be reallocated and copy it's members around)
 		uint8* Memory;
 
+		// The attributes (read from the FT_Face, but split into a more usable structure)
+		TSet<FName> Attributes;
+
 		FFontFaceAndMemory()
 			: Face(nullptr)
 			, Memory(nullptr)
+			, Attributes()
 		{
 		}
 	};
 
-	/** Mapping of font names to freetype faces */
-	TMap<FName,FFontFaceAndMemory> FontFaceMap;
+	/** Mapping of font data to freetype faces */
+	TMap<const FFontData*,FFontFaceAndMemory> FontFaceMap;
 	/** Mapping of fonts to maps of kerning pairs */
-	TMap<FSlateFontInfo, TMap<FKerningPair,int32> > FontToKerningPairMap; 
+	TMap<FSlateFontInfo, TMap<FKerningPair,int32> > FontToKerningPairMap;
+	/** Mapping of composite fonts to their cached lookup data */
+	TMap<const FCompositeFont*, TSharedPtr<FCachedCompositeFontData>> CompositeFontToCachedDataMap;
 	/** Free type library interface */
 	FT_Library FTLibrary;
 	FT_Memory CustomMemory;
@@ -523,7 +852,7 @@ int8 FKerningTable::GetKerning( TCHAR FirstChar, TCHAR SecondChar )
 		// If the kerning value hasn't been accessed yet, get the value from the cache now 
 		if( OutKerning == MAX_int8 )
 		{
-			OutKerning = FontCache.GetKerning( FirstChar, SecondChar, FontKey );
+			OutKerning = FontCache.GetKerning( FirstChar, SecondChar, FontKey.FontInfo, FontKey.Scale );
 			DirectAccessTable[Index] = OutKerning;
 		}
 	}
@@ -535,7 +864,7 @@ int8 FKerningTable::GetKerning( TCHAR FirstChar, TCHAR SecondChar )
 		int8* FoundKerning = MappedKerningPairs.Find( KerningPair );
 		if( !FoundKerning )
 		{
-			OutKerning = FontCache.GetKerning( FirstChar, SecondChar, FontKey );
+			OutKerning = FontCache.GetKerning( FirstChar, SecondChar, FontKey.FontInfo, FontKey.Scale );
 
 #if STATS
 			const uint32 CurrentMemoryUsage = MappedKerningPairs.GetAllocatedSize();
@@ -570,12 +899,22 @@ FCharacterList::FCharacterList( const FSlateFontKey& InFontKey, const FSlateFont
 	: KerningTable( InFontKey, InFontCache )
 	, FontKey( InFontKey )
 	, FontCache( InFontCache )
+	, CompositeFontHistoryRevision( 0 )
 	, MaxDirectIndexedEntries( FontCacheConstants::DirectAccessSize )
 	, MaxHeight( 0 )
 	, Baseline( 0 )
 {
+	const FCompositeFont* const CompositeFont = InFontKey.FontInfo.GetCompositeFont();
+	if( CompositeFont )
+	{
+		CompositeFontHistoryRevision = CompositeFont->HistoryRevision;
+	}
+}
 
-
+bool FCharacterList::IsStale() const
+{
+	const FCompositeFont* const CompositeFont = FontKey.FontInfo.GetCompositeFont();
+	return !CompositeFont || CompositeFontHistoryRevision != CompositeFont->HistoryRevision;
 }
 
 int8 FCharacterList::GetKerning( TCHAR First, TCHAR Second )
@@ -703,13 +1042,21 @@ FCharacterList& FSlateFontCache::GetCharacterList( const FSlateFontInfo &InFontI
 
 	TSharedRef< class FCharacterList >* CachedCharacterList = FontToCharacterListCache.Find( FontKey );
 
-	// Additional setup
-	if( CachedCharacterList == nullptr )
+	if( CachedCharacterList )
 	{
-		return FontToCharacterListCache.Add( FontKey, MakeShareable( new FCharacterList( FontKey, *this ) ) ).Get();
+		// Clear out this entry if it's stale so that we make a new one
+		if( (*CachedCharacterList)->IsStale() )
+		{
+			FontToCharacterListCache.Remove( FontKey );
+			FTInterface->Flush();
+		}
+		else
+		{
+			return CachedCharacterList->Get();
+		}
 	}
 
-	return CachedCharacterList->Get();
+	return FontToCharacterListCache.Add( FontKey, MakeShareable( new FCharacterList( FontKey, *this ) ) ).Get();
 }
 
 uint16 FSlateFontCache::GetMaxCharacterHeight( const FSlateFontInfo& InFontInfo, float FontScale ) const
@@ -738,14 +1085,77 @@ uint16 FSlateFontCache::GetBaseline( const FSlateFontInfo& InFontInfo, float Fon
 	return NewRenderData.MeasureInfo.GlobalDescender;
 }
 
-int8 FSlateFontCache::GetKerning( TCHAR First, TCHAR Second, const FSlateFontKey& FontKey ) const 
+int8 FSlateFontCache::GetKerning( TCHAR First, TCHAR Second, const FSlateFontInfo& InFontInfo, float Scale ) const 
 {
-	return FTInterface->GetKerning( First, Second, FontKey );
+	return FTInterface->GetKerning( First, Second, InFontInfo, Scale );
 }
 
 bool FSlateFontCache::HasKerning( const FSlateFontInfo& InFontInfo ) const
 {
-	return FTInterface->HasKerning( InFontInfo );
+	const FCompositeFont* const CompositeFont = InFontInfo.GetCompositeFont();
+	if(CompositeFont)
+	{
+		if (HasKerning(CompositeFont->DefaultTypeface))
+		{
+			return true;
+		}
+
+		for (const FCompositeSubFont& SubTypeface : CompositeFont->SubTypefaces)
+		{
+			if (HasKerning(SubTypeface.Typeface))
+			{
+				return true;
+			}
+		}
+	}
+
+	return false;
+}
+
+bool FSlateFontCache::HasKerning( const FTypeface& InTypeface ) const
+{
+	for (const FTypefaceEntry& TypefaceEntry : InTypeface.Fonts)
+	{
+		if (HasKerning(TypefaceEntry.Font))
+		{
+			return true;
+		}
+	}
+
+	return false;
+}
+
+bool FSlateFontCache::HasKerning( const FFontData& InFontData ) const
+{
+	return FTInterface->HasKerning( InFontData );
+}
+
+const TSet<FName>& FSlateFontCache::GetFontAttributes( const FFontData& InFontData ) const
+{
+	return FTInterface->GetFontAttributes( InFontData );
+}
+
+void FSlateFontCache::FlushObject( const UObject* const InObject )
+{
+	if( !InObject )
+	{
+		return;
+	}
+
+	bool bHasRemovedEntries = false;
+	for( auto It = FontToCharacterListCache.CreateIterator(); It; ++It )
+	{
+		if( It.Key().FontInfo.FontObject == InObject)
+		{
+			bHasRemovedEntries = true;
+			It.RemoveCurrent();
+		}
+	}
+
+	if( bHasRemovedEntries )
+	{
+		FTInterface->Flush();
+	}
 }
 
 void FSlateFontCache::ConditionalFlushCache()
diff --git a/Engine/Source/Runtime/SlateCore/Private/Fonts/FontMeasure.cpp b/Engine/Source/Runtime/SlateCore/Private/Fonts/FontMeasure.cpp
--- a/Engine/Source/Runtime/SlateCore/Private/Fonts/FontMeasure.cpp
+++ b/Engine/Source/Runtime/SlateCore/Private/Fonts/FontMeasure.cpp
@@ -2,6 +2,7 @@
 
 #include "SlateCorePrivatePCH.h"
 
+#define USE_MEASURE_CACHING 1
 
 namespace FontMeasureConstants
 {
@@ -9,6 +10,31 @@ namespace FontMeasureConstants
 	const uint32 MeasureCacheSize = 500;
 }
 
+struct FSlateFontMeasureCache
+{
+	FSlateFontMeasureCache( const FCompositeFont* const InCompositeFont )
+		: MeasureCache( FontMeasureConstants::MeasureCacheSize )
+		, CompositeFontHistoryRevision( 0 )
+	{
+		if( InCompositeFont )
+		{
+			CompositeFontHistoryRevision = InCompositeFont->HistoryRevision;
+		}
+	}
+
+	/** Check to see if our cached measure data is potentially stale for the given font */
+	bool IsStale( const FCompositeFont* const InCompositeFont ) const
+	{
+		return !InCompositeFont || CompositeFontHistoryRevision != InCompositeFont->HistoryRevision;
+	}
+
+	/** Internal measure cache */
+	FMeasureCache MeasureCache;
+
+	/** The history revision of the cached composite font */
+	int32 CompositeFontHistoryRevision;
+};
+
 TSharedRef< FSlateFontMeasure > FSlateFontMeasure::Create( const TSharedRef<class FSlateFontCache>& FontCache )
 {
 	return MakeShareable( new FSlateFontMeasure( FontCache ) );
@@ -115,29 +141,21 @@ FVector2D FSlateFontMeasure::MeasureStringInternal( const FString& Text, int32 S
 		return FVector2D( 0, MaxHeight );
 	}
 
-#define USE_MEASURE_CACHING 1
 #if USE_MEASURE_CACHING
-	TSharedPtr< FMeasureCache > CurrentMeasureCache = nullptr;
+	FMeasureCache* CurrentMeasureCache = nullptr;
 	// Do not cache strings which have small sizes or which have complicated measure requirements
 	if( DoesStartAtBeginning && DoesFinishAtEnd && !IncludeKerningWithPrecedingChar && TextRangeLength > 5 && StopAfterHorizontalOffset == INDEX_NONE )
 	{
-		FSlateFontKey FontKey(InFontInfo,FontScale);
-		TSharedPtr< FMeasureCache > FoundMeasureCache = FontToMeasureCache.FindRef( FontKey );
+		CurrentMeasureCache = FindOrAddMeasureCache(InFontInfo, FontScale);
 
-		if ( FoundMeasureCache.IsValid() )
+		if ( CurrentMeasureCache )
 		{
-			CurrentMeasureCache = FoundMeasureCache;
 			const FVector2D* CachedMeasurement = CurrentMeasureCache->AccessItem( Text );
 			if( CachedMeasurement )
 			{
 				return *CachedMeasurement;
 			}
 		}
-		else
-		{
-			CurrentMeasureCache = MakeShareable( new FMeasureCache( FontMeasureConstants::MeasureCacheSize ) );
-			FontToMeasureCache.Add( FontKey, CurrentMeasureCache );
-		}
 	}
 #endif
 
@@ -244,7 +262,7 @@ FVector2D FSlateFontMeasure::MeasureStringInternal( const FString& Text, int32 S
 	OutLastCharacterIndex = CharIndex;
 
 #if USE_MEASURE_CACHING
-	if( StopAfterHorizontalOffset == INDEX_NONE && CurrentMeasureCache.IsValid() )
+	if( StopAfterHorizontalOffset == INDEX_NONE && CurrentMeasureCache )
 	{
 		CurrentMeasureCache->Add( Text, Size );
 	}
@@ -252,6 +270,35 @@ FVector2D FSlateFontMeasure::MeasureStringInternal( const FString& Text, int32 S
 	return Size;
 }
 
+FMeasureCache* FSlateFontMeasure::FindOrAddMeasureCache( const FSlateFontInfo& InFontInfo, const float InFontScale ) const
+{
+#if USE_MEASURE_CACHING
+	FSlateFontKey FontKey(InFontInfo, InFontScale);
+	const FCompositeFont* const CompositeFont = InFontInfo.GetCompositeFont();
+
+	TSharedPtr< FSlateFontMeasureCache > FoundMeasureCache = FontToMeasureCache.FindRef( FontKey );
+	if( FoundMeasureCache.IsValid() )
+	{
+		// Clear out this entry if it's stale so that we make a new one
+		if( FoundMeasureCache->IsStale( CompositeFont ) )
+		{
+			FoundMeasureCache.Reset();
+			FontToMeasureCache.Remove( FontKey );
+		}
+		else
+		{
+			return &FoundMeasureCache->MeasureCache;
+		}
+	}
+	
+	FoundMeasureCache = MakeShareable( new FSlateFontMeasureCache( CompositeFont ) );
+	FontToMeasureCache.Add( FontKey, FoundMeasureCache );
+	return &FoundMeasureCache->MeasureCache;
+#endif
+
+	return nullptr;
+}
+
 uint16 FSlateFontMeasure::GetMaxCharacterHeight( const FSlateFontInfo& InFontInfo, float FontScale ) const
 {
 	FCharacterList& CharacterList = FontCache->GetCharacterList( InFontInfo, FontScale );
diff --git a/Engine/Source/Runtime/SlateCore/Private/Fonts/FontProviderInterface.cpp b/Engine/Source/Runtime/SlateCore/Private/Fonts/FontProviderInterface.cpp
new file mode 100644
--- /dev/null
+++ b/Engine/Source/Runtime/SlateCore/Private/Fonts/FontProviderInterface.cpp
@@ -0,0 +1,9 @@
+// Copyright 1998-2014 Epic Games, Inc. All Rights Reserved.
+
+#include "SlateCorePrivatePCH.h"
+#include "FontProviderInterface.h"
+
+UFontProviderInterface::UFontProviderInterface(const FObjectInitializer& ObjectInitializer)
+	: Super(ObjectInitializer)
+{
+}
diff --git a/Engine/Source/Runtime/SlateCore/Private/Fonts/LegacySlateFontInfoCache.cpp b/Engine/Source/Runtime/SlateCore/Private/Fonts/LegacySlateFontInfoCache.cpp
new file mode 100644
--- /dev/null
+++ b/Engine/Source/Runtime/SlateCore/Private/Fonts/LegacySlateFontInfoCache.cpp
@@ -0,0 +1,92 @@
+// Copyright 1998-2014 Epic Games, Inc. All Rights Reserved.
+
+#include "SlateCorePrivatePCH.h"
+#include "LegacySlateFontInfoCache.h"
+
+TSharedPtr<FLegacySlateFontInfoCache> FLegacySlateFontInfoCache::Singleton;
+
+FLegacySlateFontInfoCache& FLegacySlateFontInfoCache::Get()
+{
+	if (!Singleton.IsValid())
+	{
+		Singleton = MakeShareable(new FLegacySlateFontInfoCache());
+	}
+	return *Singleton;
+}
+
+TSharedPtr<const FCompositeFont> FLegacySlateFontInfoCache::GetCompositeFont(const FName& InLegacyFontName, const EFontHinting InLegacyFontHinting)
+{
+	if (InLegacyFontName.IsNone())
+	{
+		return nullptr;
+	}
+
+	FString LegacyFontPath = InLegacyFontName.ToString();
+
+	// Work out what the given path is supposed to be relative to
+	if (!FPaths::FileExists(LegacyFontPath))
+	{
+		// UMG assets specify the path either relative to the game or engine content directories - test both
+		LegacyFontPath = FPaths::GameContentDir() / InLegacyFontName.ToString();
+		if (!FPaths::FileExists(LegacyFontPath))
+		{
+			LegacyFontPath = FPaths::EngineContentDir() / InLegacyFontName.ToString();
+			if (!FPaths::FileExists(LegacyFontPath))
+			{
+				// Missing font file - just use what we were given
+				LegacyFontPath = InLegacyFontName.ToString();
+			}
+		}
+	}
+
+	const FLegacyFontKey LegacyFontKey(*LegacyFontPath, InLegacyFontHinting);
+
+	{
+		TSharedPtr<const FCompositeFont>* const ExistingCompositeFont = LegacyFontNameToCompositeFont.Find(LegacyFontKey);
+		if(ExistingCompositeFont)
+		{
+			return *ExistingCompositeFont;
+		}
+	}
+
+	TSharedRef<const FCompositeFont> NewCompositeFont = MakeShareable(new FCompositeFont(NAME_None, LegacyFontPath, InLegacyFontHinting));
+	LegacyFontNameToCompositeFont.Add(LegacyFontKey, NewCompositeFont);
+	return NewCompositeFont;
+}
+
+TSharedPtr<const FCompositeFont> FLegacySlateFontInfoCache::GetSystemFont()
+{
+	if (!SystemFont.IsValid())
+	{
+		TArray<uint8> FontBytes = FPlatformMisc::GetSystemFontBytes();
+		SystemFont = MakeShareable(new FCompositeFont(NAME_None, TEXT("DefaultSystemFont"), MoveTemp(FontBytes), EFontHinting::Default));
+	}
+	return SystemFont;
+}
+
+const FFontData& FLegacySlateFontInfoCache::GetFallbackFont()
+{
+	const FName FallbackFontName = *(FPaths::EngineContentDir() / TEXT("Slate/Fonts/") / (NSLOCTEXT("Slate", "FallbackFont", "DroidSansFallback").ToString() + TEXT(".ttf")));
+
+	{
+		TSharedPtr<const FFontData>* const ExistingFallbackFont = FallbackFonts.Find(FallbackFontName);
+		if(ExistingFallbackFont)
+		{
+			return **ExistingFallbackFont;
+		}
+	}
+
+	TSharedRef<const FFontData> NewFallbackFont = MakeShareable(new FFontData(FallbackFontName.ToString(), EFontHinting::Default));
+	FallbackFonts.Add(FallbackFontName, NewFallbackFont);
+	return *NewFallbackFont;
+}
+
+const FFontData& FLegacySlateFontInfoCache::GetLastResortFont()
+{
+	if (!LastResortFont.IsValid())
+	{
+		const FString LastResortFontPath = FPaths::EngineContentDir() / TEXT("Slate/Fonts/LastResort.ttf");
+		LastResortFont = MakeShareable(new FFontData(LastResortFontPath, EFontHinting::Default));
+	}
+	return *LastResortFont;
+}
diff --git a/Engine/Source/Runtime/SlateCore/Private/Fonts/LegacySlateFontInfoCache.h b/Engine/Source/Runtime/SlateCore/Private/Fonts/LegacySlateFontInfoCache.h
new file mode 100644
--- /dev/null
+++ b/Engine/Source/Runtime/SlateCore/Private/Fonts/LegacySlateFontInfoCache.h
@@ -0,0 +1,81 @@
+// Copyright 1998-2014 Epic Games, Inc. All Rights Reserved.
+
+#pragma once
+
+/**
+ * Cache used to efficiently upgrade legacy FSlateFontInfo structs to use a composite font by reducing the amount of duplicate instances that are created
+ */
+class FLegacySlateFontInfoCache
+{
+public:
+
+	/**
+	 * Get (or create) the singleton instance of this cache
+	 */
+	static FLegacySlateFontInfoCache& Get();
+
+	/**
+	 * Get (or create) an appropriate composite font from the legacy font name
+	 */
+	TSharedPtr<const FCompositeFont> GetCompositeFont(const FName& InLegacyFontName, const EFontHinting InLegacyFontHinting);
+
+	/**
+	 * Get (or create) the default system font
+	 */
+	TSharedPtr<const FCompositeFont> GetSystemFont();
+
+	/**
+	 * Get (or create) the culture specific fallback font
+	 */
+	const FFontData& GetFallbackFont();
+
+	/**
+	 * Get (or create) the last resort fallback font
+	 */
+	const FFontData& GetLastResortFont();
+
+private:
+	struct FLegacyFontKey
+	{
+		FLegacyFontKey()
+			: Name()
+			, Hinting(EFontHinting::Default)
+		{
+		}
+
+		FLegacyFontKey(const FName& InName, const EFontHinting InHinting)
+			: Name(InName)
+			, Hinting(InHinting)
+		{
+		}
+
+		bool operator==(const FLegacyFontKey& Other) const 
+		{
+			return Name == Other.Name && Hinting == Other.Hinting;
+		}
+
+		bool operator!=(const FLegacyFontKey& Other) const 
+		{
+			return !(*this == Other);
+		}
+
+		friend inline uint32 GetTypeHash(const FLegacyFontKey& Key)
+		{
+			uint32 Hash = 0;
+			Hash = HashCombine(Hash, GetTypeHash(Key.Name));
+			Hash = HashCombine(Hash, uint32(Key.Hinting));
+			return Hash;
+		}
+
+		FName Name;
+		EFontHinting Hinting;
+	};
+
+	TMap<FLegacyFontKey, TSharedPtr<const FCompositeFont>> LegacyFontNameToCompositeFont;
+	TSharedPtr<const FCompositeFont> SystemFont;
+	
+	TMap<FName, TSharedPtr<const FFontData>> FallbackFonts;
+	TSharedPtr<const FFontData> LastResortFont;
+
+	static TSharedPtr<FLegacySlateFontInfoCache> Singleton;
+};
diff --git a/Engine/Source/Runtime/SlateCore/Private/Fonts/SlateFontInfo.cpp b/Engine/Source/Runtime/SlateCore/Private/Fonts/SlateFontInfo.cpp
--- a/Engine/Source/Runtime/SlateCore/Private/Fonts/SlateFontInfo.cpp
+++ b/Engine/Source/Runtime/SlateCore/Private/Fonts/SlateFontInfo.cpp
@@ -1,53 +1,144 @@
 // Copyright 1998-2014 Epic Games, Inc. All Rights Reserved.
 
 #include "SlateCorePrivatePCH.h"
+#include "FontProviderInterface.h"
+#include "LegacySlateFontInfoCache.h"
 
 
 /* FSlateFontInfo structors
  *****************************************************************************/
 
 FSlateFontInfo::FSlateFontInfo( )
-	: FontName(NAME_None)
+	: FontObject(nullptr)
+	, CompositeFont()
+	, TypefaceFontName()
 	, Size(0)
-	, Hinting(EFontHinting::Default)
-{ }
+	, FontName_DEPRECATED()
+	, Hinting_DEPRECATED(EFontHinting::Default)
+{
+}
+
+
+FSlateFontInfo::FSlateFontInfo( TSharedPtr<const FCompositeFont> InCompositeFont, const int32 InSize, const FName& InTypefaceFontName )
+	: FontObject(nullptr)
+	, CompositeFont(InCompositeFont)
+	, TypefaceFontName(InTypefaceFontName)
+	, Size(InSize)
+	, FontName_DEPRECATED()
+	, Hinting_DEPRECATED(EFontHinting::Default)
+{
+}
+
+
+FSlateFontInfo::FSlateFontInfo( const UObject* InFontObject, const int32 InSize, const FName& InTypefaceFontName )
+	: FontObject(InFontObject)
+	, CompositeFont()
+	, TypefaceFontName(InTypefaceFontName)
+	, Size(InSize)
+	, FontName_DEPRECATED()
+	, Hinting_DEPRECATED(EFontHinting::Default)
+{
+}
 
 
 FSlateFontInfo::FSlateFontInfo( const FString& InFontName, uint16 InSize, EFontHinting InHinting )
-	: FontName(*InFontName)
+	: FontObject(nullptr)
+	, CompositeFont()
+	, TypefaceFontName()
 	, Size(InSize)
-	, Hinting(InHinting)
+	, FontName_DEPRECATED(*InFontName)
+	, Hinting_DEPRECATED(InHinting)
 {
 	//Useful for debugging style breakages
 	//check( FPaths::FileExists( FontName.ToString() ) );
+
+	UpgradeLegacyFontInfo();
 }
 
 
 FSlateFontInfo::FSlateFontInfo( const FName& InFontName, uint16 InSize, EFontHinting InHinting )
-	: FontName(InFontName)
+	: FontObject(nullptr)
+	, CompositeFont()
+	, TypefaceFontName()
 	, Size(InSize)
-	, Hinting(InHinting)
+	, FontName_DEPRECATED(InFontName)
+	, Hinting_DEPRECATED(InHinting)
 {
 	//Useful for debugging style breakages
 	//check( FPaths::FileExists( FontName.ToString() ) );
+
+	UpgradeLegacyFontInfo();
 }
 
 
 FSlateFontInfo::FSlateFontInfo( const ANSICHAR* InFontName, uint16 InSize, EFontHinting InHinting )
-	: FontName(InFontName)
+	: FontObject(nullptr)
+	, CompositeFont()
+	, TypefaceFontName()
 	, Size(InSize)
-	, Hinting(InHinting)
+	, FontName_DEPRECATED(InFontName)
+	, Hinting_DEPRECATED(InHinting)
 {
 	//Useful for debugging style breakages
 	//check( FPaths::FileExists( FontName.ToString() ) );
+
+	UpgradeLegacyFontInfo();
 }
 
 
 FSlateFontInfo::FSlateFontInfo( const WIDECHAR* InFontName, uint16 InSize, EFontHinting InHinting )
-	: FontName(InFontName)
+	: FontObject(nullptr)
+	, CompositeFont()
+	, TypefaceFontName()
 	, Size(InSize)
-	, Hinting(InHinting)
+	, FontName_DEPRECATED(InFontName)
+	, Hinting_DEPRECATED(InHinting)
 {
 	//Useful for debugging style breakages
 	//check( FPaths::FileExists( FontName.ToString() ) );
+
+	UpgradeLegacyFontInfo();
+}
+
+
+bool FSlateFontInfo::HasValidFont() const
+{
+	return CompositeFont.IsValid() || FontObject != nullptr;
+}
+
+
+const FCompositeFont* FSlateFontInfo::GetCompositeFont() const
+{
+	const IFontProviderInterface* FontProvider = Cast<const IFontProviderInterface>(FontObject);
+	if (FontProvider)
+	{
+		return FontProvider->GetCompositeFont();
+	}
+
+	if (CompositeFont.IsValid())
+	{
+		return CompositeFont.Get();
+	}
+
+	return nullptr;
+}
+
+
+void FSlateFontInfo::PostSerialize(const FArchive& Ar)
+{
+	if (Ar.UE4Ver() < VER_UE4_SLATE_COMPOSITE_FONTS && !FontObject)
+	{
+		UpgradeLegacyFontInfo();
+	}
+}
+
+
+void FSlateFontInfo::UpgradeLegacyFontInfo()
+{
+	static const FName SpecialName_DefaultSystemFont("DefaultSystemFont");
+
+	// Special case for using the default system font
+	CompositeFont = (FontName_DEPRECATED == SpecialName_DefaultSystemFont)
+		? FLegacySlateFontInfoCache::Get().GetSystemFont()
+		: FLegacySlateFontInfoCache::Get().GetCompositeFont(FontName_DEPRECATED, Hinting_DEPRECATED);
 }
diff --git a/Engine/Source/Runtime/SlateCore/Private/Rendering/ElementBatcher.cpp b/Engine/Source/Runtime/SlateCore/Private/Rendering/ElementBatcher.cpp
--- a/Engine/Source/Runtime/SlateCore/Private/Rendering/ElementBatcher.cpp
+++ b/Engine/Source/Runtime/SlateCore/Private/Rendering/ElementBatcher.cpp
@@ -560,7 +560,7 @@ void FSlateElementBatcher::AddTextElement(const FSlateDrawElement& DrawElement)
 				// Font has a new texture for this glyph. Refresh the batch we use and the index we are currently using
 				FontTextureIndex = Entry.TextureIndex;
 
-				FontTexture = FontCache.GetTextureResource( FontTextureIndex );
+				FontTexture = FontCache.GetSlateTextureResource( FontTextureIndex );
 				ElementBatch = &FindBatchForElement( Layer, FShaderParams(), FontTexture, ESlateDrawPrimitive::TriangleList, ESlateShader::Font, InDrawEffects, ESlateBatchDrawFlag::None, DrawElement.GetScissorRect() );
 
 				BatchVertices = &BatchVertexArrays[ElementBatch->VertexArrayIndex];
diff --git a/Engine/Source/Runtime/SlateCore/Private/Styling/SlateTypes.cpp b/Engine/Source/Runtime/SlateCore/Private/Styling/SlateTypes.cpp
--- a/Engine/Source/Runtime/SlateCore/Private/Styling/SlateTypes.cpp
+++ b/Engine/Source/Runtime/SlateCore/Private/Styling/SlateTypes.cpp
@@ -334,7 +334,7 @@ FSearchBoxStyle::FSearchBoxStyle()
 FSearchBoxStyle& FSearchBoxStyle::SetTextBoxStyle( const FEditableTextBoxStyle& InTextBoxStyle )
 { 
 	TextBoxStyle = InTextBoxStyle;
-	if (ActiveFontInfo.FontName.IsNone())
+	if (!ActiveFontInfo.HasValidFont())
 	{
 		ActiveFontInfo = TextBoxStyle.Font;
 	}
diff --git a/Engine/Source/Runtime/SlateCore/Public/Fonts/CompositeFont.h b/Engine/Source/Runtime/SlateCore/Public/Fonts/CompositeFont.h
new file mode 100644
--- /dev/null
+++ b/Engine/Source/Runtime/SlateCore/Public/Fonts/CompositeFont.h
@@ -0,0 +1,222 @@
+// Copyright 1998-2014 Epic Games, Inc. All Rights Reserved.
+
+#pragma once
+
+#include "CompositeFont.generated.h"
+
+UENUM()
+enum class EFontHinting : uint8
+{
+	/** Use the default hinting specified in the font */
+	Default,
+	/** Force the use of an automatic hinting algorithm */
+	Auto,
+	/** Force the use of an automatic light hinting algorithm, optimized for non-monochrome displays */
+	AutoLight,
+	/** Force the use of an automatic hinting algorithm optimized for monochrome displays */
+	Monochrome,
+	/** Do not use hinting */
+	None,
+};
+
+/** Payload data describing an individual font in a typeface */
+USTRUCT(BlueprintType)
+struct SLATECORE_API FFontData
+{
+	GENERATED_USTRUCT_BODY()
+
+	/** Default constructor */
+	FFontData()
+		: FontFilename()
+		, FontData()
+		, Hinting(EFontHinting::Default)
+	{
+	}
+
+	/** Construct the raw font data from a filename (it will load the file into FontData) */
+	FFontData(FString InFontFilename, const EFontHinting InHinting);
+
+	/** Construct the raw font data from a filename, and the data associated with that file */
+	FFontData(FString InFontFilename, TArray<uint8> InFontData, const EFontHinting InHinting)
+		: FontFilename(MoveTemp(InFontFilename))
+		, FontData(MoveTemp(InFontData))
+		, Hinting(InHinting)
+	{
+	}
+
+	/** Set the new font data from a filename (it will load the file into FontData) */
+	bool SetFont(FString InFontFilename);
+
+	/** Set the new font data from a filename, and the data associated with that file */
+	void SetFont(FString InFontFilename, TArray<uint8> InFontData);
+
+	/** The filename of the font to use - this may not always exist on disk, as we may have previously loaded and cached the font data inside an asset */
+	UPROPERTY(EditAnywhere, BlueprintReadWrite, Category=Font)
+	FString FontFilename;
+
+	/** The data associated with the font - this should always be filled in providing the source font filename is valid */
+	UPROPERTY()
+	TArray<uint8> FontData;
+
+	/** The hinting algorithm to use with the font */
+	UPROPERTY(EditAnywhere, BlueprintReadWrite, Category=Font)
+	EFontHinting Hinting;
+};
+
+/** A single entry in a typeface */
+USTRUCT(BlueprintType)
+struct SLATECORE_API FTypefaceEntry
+{
+	GENERATED_USTRUCT_BODY()
+
+	/** Default constructor */
+	FTypefaceEntry()
+	{
+	}
+
+	/** Construct the entry from a name */
+	FTypefaceEntry(const FName& InFontName)
+		: Name(InFontName)
+	{
+	}
+
+	/** Construct the entry from a filename */
+	FTypefaceEntry(const FName& InFontName, FString InFontFilename, const EFontHinting InHinting)
+		: Name(InFontName)
+		, Font(MoveTemp(InFontFilename), InHinting)
+	{
+	}
+
+	/** Construct the entry from a filename, and the data associated with that file */
+	FTypefaceEntry(const FName& InFontName, FString InFontFilename, TArray<uint8> InFontData, const EFontHinting InHinting)
+		: Name(InFontName)
+		, Font(MoveTemp(InFontFilename), MoveTemp(InFontData), InHinting)
+	{
+	}
+
+	/** Name used to identify this font within its typeface */
+	UPROPERTY(EditAnywhere, BlueprintReadWrite, Category=Font)
+	FName Name;
+
+	/** Raw font data for this font */
+	UPROPERTY(EditAnywhere, BlueprintReadWrite, Category=Font)
+	FFontData Font;
+};
+
+/** Definition for a typeface (a family of fonts) */
+USTRUCT(BlueprintType)
+struct SLATECORE_API FTypeface
+{
+	GENERATED_USTRUCT_BODY()
+
+	/** Default constructor */
+	FTypeface()
+	{
+	}
+
+	/** Convenience constructor for when your font family only contains a single font */
+	FTypeface(const FName& InFontName, FString InFontFilename, const EFontHinting InHinting)
+	{
+		Fonts.Emplace(FTypefaceEntry(InFontName, MoveTemp(InFontFilename), InHinting));
+	}
+
+	/** Convenience constructor for when your font family only contains a single font */
+	FTypeface(const FName& InFontName, FString InFontFilename, TArray<uint8> InFontData, const EFontHinting InHinting)
+	{
+		Fonts.Emplace(FTypefaceEntry(InFontName, MoveTemp(InFontFilename), MoveTemp(InFontData), InHinting));
+	}
+
+	/** Append a new font into this family */
+	FTypeface& AppendFont(const FName& InFontName, FString InFontFilename, const EFontHinting InHinting = EFontHinting::Default)
+	{
+		Fonts.Emplace(FTypefaceEntry(InFontName, MoveTemp(InFontFilename), InHinting));
+		return *this;
+	}
+
+	/** Append a new font into this family */
+	FTypeface& AppendFont(const FName& InFontName, FString InFontFilename, TArray<uint8> InFontData, const EFontHinting InHinting = EFontHinting::Default)
+	{
+		Fonts.Emplace(FTypefaceEntry(InFontName, MoveTemp(InFontFilename), MoveTemp(InFontData), InHinting));
+		return *this;
+	}
+
+	/** The fonts contained within this family */
+	UPROPERTY(EditAnywhere, BlueprintReadWrite, Category=Font)
+	TArray<FTypefaceEntry> Fonts;
+};
+
+USTRUCT(BlueprintType)
+struct SLATECORE_API FCompositeSubFont
+{
+	GENERATED_USTRUCT_BODY()
+
+	/** Default constructor */
+	FCompositeSubFont()
+		: Typeface()
+		, CharacterRanges()
+		, ScalingFactor(1.0f)
+	{
+	}
+
+	/** Typeface data for this sub-font */
+	UPROPERTY(EditAnywhere, BlueprintReadWrite, Category=Font)
+	FTypeface Typeface;
+
+	/** Array of character ranges for which this sub-font should be used */
+	UPROPERTY(EditAnywhere, BlueprintReadWrite, Category=Font)
+	TArray<FInt32Range> CharacterRanges;
+
+	/** Amount to scale this sub-font so that it better matches the size of the default font */
+	UPROPERTY(EditAnywhere, BlueprintReadWrite, Category=Font)
+	float ScalingFactor;
+};
+
+USTRUCT(BlueprintType)
+struct SLATECORE_API FCompositeFont
+{
+	GENERATED_USTRUCT_BODY()
+
+	/** Default constructor */
+	FCompositeFont()
+		: DefaultTypeface()
+		, SubTypefaces()
+		, HistoryRevision(0)
+	{
+	}
+
+	/** Convenience constructor for when your composite font only contains a single font */
+	FCompositeFont(const FName& InFontName, FString InFontFilename, const EFontHinting InHinting = EFontHinting::Default)
+		: DefaultTypeface(InFontName, MoveTemp(InFontFilename), InHinting)
+		, SubTypefaces()
+		, HistoryRevision(0)
+	{
+	}
+
+	/** Convenience constructor for when your composite font only contains a single font */
+	FCompositeFont(const FName& InFontName, FString InFontFilename, TArray<uint8> InFontData, const EFontHinting InHinting = EFontHinting::Default)
+		: DefaultTypeface(InFontName, MoveTemp(InFontFilename), MoveTemp(InFontData), InHinting)
+		, SubTypefaces()
+		, HistoryRevision(0)
+	{
+	}
+
+	/** Call this when the composite font is changed after its initial setup - this allows various caches to update as required */
+	void MakeDirty()
+	{
+		++HistoryRevision;
+	}
+
+	/** The default typeface that will be used when not overridden by a sub-typeface */
+	UPROPERTY(EditAnywhere, BlueprintReadWrite, Category=RuntimeFont)
+	FTypeface DefaultTypeface;
+
+	/** Sub-typefaces to use for a specific set of characters */
+	UPROPERTY(EditAnywhere, BlueprintReadWrite, Category=RuntimeFont)
+	TArray<FCompositeSubFont> SubTypefaces;
+
+	/** 
+	 * Transient value containing the current history ID of this composite font
+	 * This should be updated when the composite font is changed (which should happen infrequently as composite fonts are assumed to be mostly immutable) once they've been setup
+	 */
+	int32 HistoryRevision;
+};
diff --git a/Engine/Source/Runtime/SlateCore/Public/Fonts/FontCache.h b/Engine/Source/Runtime/SlateCore/Public/Fonts/FontCache.h
--- a/Engine/Source/Runtime/SlateCore/Public/Fonts/FontCache.h
+++ b/Engine/Source/Runtime/SlateCore/Public/Fonts/FontCache.h
@@ -7,7 +7,7 @@ class FFreeTypeInterface;
 
 
 /** Information for rendering one character */
-struct FCharacterEntry
+struct SLATECORE_API FCharacterEntry
 {
 	/** Start X location of the character in the texture */
 	float StartU;
@@ -46,7 +46,7 @@ struct FCharacterEntry
 	}
 };
 
-struct FKerningPair
+struct SLATECORE_API FKerningPair
 {
 	TCHAR First;
 	TCHAR Second;
@@ -73,7 +73,7 @@ class FSlateFontCache;
 /**
  * A Kerning table for a single font key
  */
-class FKerningTable
+class SLATECORE_API FKerningTable
 {
 public:
 	FKerningTable( const FSlateFontKey& InFont, const FSlateFontCache& InFontCache );
@@ -112,7 +112,7 @@ private:
  * Every character indexed by TCHAR could potentially cost a lot of memory of a lot of empty entries are created
  * because characters being used are far apart
  */
-class FCharacterList
+class SLATECORE_API FCharacterList
 {
 public:
 	FCharacterList( const FSlateFontKey& InFontKey, const FSlateFontCache& InFontCache );
@@ -166,6 +166,9 @@ public:
 		}
 	}
 
+	/** Check to see if our cached data is potentially stale for our font */
+	bool IsStale() const;
+
 	/**
 	 * Gets a kerning value for a pair of characters
 	 *
@@ -209,6 +212,8 @@ private:
 	FSlateFontKey FontKey;
 	/** Reference to the font cache for accessing new unseen characters */
 	const class FSlateFontCache& FontCache;
+	/** The history revision of the cached composite font */
+	int32 CompositeFontHistoryRevision;
 	/** Number of directly indexed entries */
 	int32 MaxDirectIndexedEntries;
 	/** The global max height for any character in this font */
@@ -253,6 +258,11 @@ public:
 	 */
 	bool AddNewEntry( TCHAR Character, const FSlateFontKey& InKey, FCharacterEntry& OutCharacterEntry ) const;
 
+	/**
+	 * Flush the given object out of the cache
+	 */
+	void FlushObject( const UObject* const InObject );
+
 	/** 
 	 * Flush the cache if needed
 	 */
@@ -274,23 +284,34 @@ public:
 	 * @param Index	The index of the texture 
 	 * @return Handle to the texture resource
 	 */
-	FSlateShaderResource* GetTextureResource( uint32 Index ) { return FontAtlases[Index]->GetTexture(); }
+	class FSlateShaderResource* GetSlateTextureResource( uint32 Index ) { return FontAtlases[Index]->GetSlateTexture(); }
+	class FTextureResource* GetEngineTextureResource( uint32 Index ) { return FontAtlases[Index]->GetEngineTexture(); }
 
 	/**
 	 * Calculates the kerning amount for a pair of characters
 	 *
-	 * @param First		The first character in the pair
-	 * @param Second	The second character in the pair
-	 * @oaran FontKey	The font to get kerning from
+	 * @param First			The first character in the pair
+	 * @param Second		The second character in the pair
+	 * @param InFontInfo	Information about the font that used to draw the string with the first and second characters
 	 * @return The kerning amount, 0 if no kerning
 	 */
-	int8 GetKerning( TCHAR First, TCHAR Second, const FSlateFontKey& FontKey ) const;
+	int8 GetKerning( TCHAR First, TCHAR Second, const FSlateFontInfo& InFontInfo, float Scale ) const;
 
 	/**
-	 * @return Whether or not a font has kerning information
+	 * @return Whether or not the font has kerning information (this checks both the default font as well as any sub-fonts)
 	 */
 	bool HasKerning( const FSlateFontInfo& InFontInfo ) const;
 
+	/**
+	 * @return Whether or not the typeface has kerning information
+	 */
+	bool HasKerning( const FTypeface& InTypeface ) const;
+
+	/**
+	 * @return Whether or not the font has kerning information
+	 */
+	bool HasKerning( const FFontData& InFontData ) const;
+
 	/**
 	 * Returns the height of the largest character in the font. 
 	 *
@@ -311,6 +332,15 @@ public:
 	 */
 	uint16 GetBaseline( const FSlateFontInfo& InFontInfo, float FontScale ) const;
 
+	/**
+	 * Returns the font attributes for the specified font.
+	 *
+	 * @param InFontData	The font to get attributes for 
+	 * 
+	 * @return The font attributes for the specified font.
+	 */
+	const TSet<FName>& GetFontAttributes( const FFontData& InFontData ) const;
+
 	/**
 	 * Clears all cached data from the cache
 	 */
diff --git a/Engine/Source/Runtime/SlateCore/Public/Fonts/FontMeasure.h b/Engine/Source/Runtime/SlateCore/Public/Fonts/FontMeasure.h
--- a/Engine/Source/Runtime/SlateCore/Public/Fonts/FontMeasure.h
+++ b/Engine/Source/Runtime/SlateCore/Public/Fonts/FontMeasure.h
@@ -2,6 +2,7 @@
 #pragma once
 
 typedef class FLRUStringCache FMeasureCache;
+struct FSlateFontMeasureCache;
 
 class SLATECORE_API FSlateFontMeasure 
 {
@@ -222,10 +223,15 @@ private:
 	 */
 	FVector2D MeasureStringInternal( const FString& Text, int32 StartIndex, int32 EndIndex, const FSlateFontInfo& InFontInfo, bool IncludeKerningWithPrecedingChar, float FontScale, int32 StopAfterHorizontalOffset, ELastCharacterIndexFormat CharIndexFormat, int32& OutLastCharacterIndex ) const;
 
+	/**
+	 * Check to see if there's an existing cached measurement, or failing that, add a new entry so that we can cache a new measurement
+	 */
+	FMeasureCache* FindOrAddMeasureCache( const FSlateFontInfo& InFontInfo, const float InFontScale ) const;
+
 private:
 
 	/** Mapping Font keys to cached data */
-	mutable TMap< FSlateFontKey, TSharedPtr< FMeasureCache > > FontToMeasureCache;
+	mutable TMap< FSlateFontKey, TSharedPtr< FSlateFontMeasureCache > > FontToMeasureCache;
 
 	TSharedRef<class FSlateFontCache> FontCache;
 };
\ No newline at end of file
diff --git a/Engine/Source/Runtime/SlateCore/Public/Fonts/FontProviderInterface.h b/Engine/Source/Runtime/SlateCore/Public/Fonts/FontProviderInterface.h
new file mode 100644
--- /dev/null
+++ b/Engine/Source/Runtime/SlateCore/Public/Fonts/FontProviderInterface.h
@@ -0,0 +1,19 @@
+// Copyright 1998-2014 Epic Games, Inc. All Rights Reserved.
+
+#pragma once
+
+#include "CompositeFont.h"
+#include "FontProviderInterface.generated.h"
+
+UINTERFACE(MinimalAPI, meta=(CannotImplementInterfaceInBlueprint))
+class UFontProviderInterface : public UInterface
+{
+	GENERATED_UINTERFACE_BODY()
+};
+
+class IFontProviderInterface
+{
+	GENERATED_IINTERFACE_BODY()
+
+	virtual const FCompositeFont* GetCompositeFont() const = 0;
+};
diff --git a/Engine/Source/Runtime/SlateCore/Public/Fonts/FontTypes.h b/Engine/Source/Runtime/SlateCore/Public/Fonts/FontTypes.h
--- a/Engine/Source/Runtime/SlateCore/Public/Fonts/FontTypes.h
+++ b/Engine/Source/Runtime/SlateCore/Public/Fonts/FontTypes.h
@@ -20,7 +20,10 @@ struct SLATECORE_API FSlateFontKey
 
 	friend inline uint32 GetTypeHash( const FSlateFontKey& Key )
 	{
-		return GetTypeHash(Key.FontInfo) + FCrc::MemCrc32( &Key.Scale, sizeof(float) );
+		uint32 Hash = 0;
+		Hash = HashCombine(Hash, GetTypeHash(Key.FontInfo));
+		Hash = HashCombine(Hash, GetTypeHash(Key.Scale));
+		return Hash;
 	}
 };
 
@@ -70,9 +73,14 @@ public:
 	virtual ~FSlateFontAtlas();
 
 	/**
-	 * Returns the texture resource used by slate
+	 * Returns the texture resource used by Slate
 	 */
-	virtual class FSlateShaderResource* GetTexture() = 0;
+	virtual class FSlateShaderResource* GetSlateTexture() = 0;
+
+	/**
+	 * Returns the texture resource used the Engine
+	 */
+	virtual class FTextureResource* GetEngineTexture() = 0;
 	
 	/**
 	 * Releases rendering resources of this cache
diff --git a/Engine/Source/Runtime/SlateCore/Public/Fonts/SlateFontInfo.h b/Engine/Source/Runtime/SlateCore/Public/Fonts/SlateFontInfo.h
--- a/Engine/Source/Runtime/SlateCore/Public/Fonts/SlateFontInfo.h
+++ b/Engine/Source/Runtime/SlateCore/Public/Fonts/SlateFontInfo.h
@@ -1,24 +1,11 @@
 // Copyright 1998-2014 Epic Games, Inc. All Rights Reserved.
 
 #pragma once
+
 #include "Templates/TypeHash.h"
+#include "CompositeFont.h"
 #include "SlateFontInfo.generated.h"
 
-UENUM()
-enum class EFontHinting : uint8
-{
-	/** Use the default hinting specified in the font */
-	Default,
-	/** Force the use of an automatic hinting algorithm */
-	Auto,
-	/** Force the use of an automatic light hinting algorithm, optimized for non-monochrome displays */
-	AutoLight,
-	/** Force the use of an automatic hinting algorithm optimized for monochrome displays */
-	Monochrome,
-	/** Do not use hinting */
-	None,
-};
-
 /**
  * A representation of a font in Slate.
  */
@@ -27,17 +14,30 @@ struct SLATECORE_API FSlateFontInfo
 {
 	GENERATED_USTRUCT_BODY()
 
-	/** The name of the font */
-	UPROPERTY(EditAnywhere, BlueprintReadWrite, Category=SlateStyleRules)
-	FName FontName;
+	/** The font object (valid when used from UMG or a Slate widget style asset) */
+	UPROPERTY(EditAnywhere, BlueprintReadWrite, Category=SlateStyleRules, meta=(AllowedClasses="Font", DisplayName="Font Family"))
+	const UObject* FontObject;
+
+	/** The composite font data to use (valid when used with a Slate style set in C++) */
+	TSharedPtr<const FCompositeFont> CompositeFont;
+
+	/** The name of the font to use from the default typeface (None will use the first entry) */
+	UPROPERTY(EditAnywhere, BlueprintReadWrite, Category=SlateStyleRules, meta=(DisplayName="Font"))
+	FName TypefaceFontName;
 
 	/** The size of the font */
 	UPROPERTY(EditAnywhere, BlueprintReadWrite, Category=SlateStyleRules)
 	int32 Size;
 
+private:
+
+	/** The name of the font */
+	UPROPERTY()
+	FName FontName_DEPRECATED;
+
 	/** The hinting algorithm to use with the font */
-	UPROPERTY(EditAnywhere, BlueprintReadWrite, Category=SlateStyleRules)
-	EFontHinting Hinting;
+	UPROPERTY()
+	EFontHinting Hinting_DEPRECATED;
 
 public:
 
@@ -45,7 +45,25 @@ public:
 	FSlateFontInfo();
 
 	/**
-	 * Creates and initializes a new instance with the specified font name and size.
+	 * Creates and initializes a new instance with the specified font, size, and emphasis.
+	 *
+	 * @param InCompositeFont The font instance to use.
+	 * @param InSize The size of the font.
+	 * @param InTypefaceFontName The name of the font to use from the default typeface (None will use the first entry)
+	 */
+	FSlateFontInfo( TSharedPtr<const FCompositeFont> InCompositeFont, const int32 InSize, const FName& InTypefaceFontName = NAME_None );
+
+	/**
+	 * Creates and initializes a new instance with the specified font, size, and emphasis.
+	 *
+	 * @param InFontObject The font instance to use.
+	 * @param InSize The size of the font.
+	 * @param InFamilyFontName The name of the font to use from the default typeface (None will use the first entry)
+	 */
+	FSlateFontInfo( const UObject* InFontObject, const int32 InSize, const FName& InTypefaceFontName = NAME_None );
+
+	/**
+	 * DEPRECATED - Creates and initializes a new instance with the specified font name and size.
 	 *
 	 * @param InFontName The name of the font.
 	 * @param InSize The size of the font.
@@ -54,7 +72,7 @@ public:
 	FSlateFontInfo( const FString& InFontName, uint16 InSize, EFontHinting InHinting = EFontHinting::Default );
 
 	/**
-	 * Creates and initializes a new instance with the specified font name and size.
+	 * DEPRECATED - Creates and initializes a new instance with the specified font name and size.
 	 *
 	 * @param InFontName The name of the font.
 	 * @param InSize The size of the font.
@@ -63,7 +81,7 @@ public:
 	FSlateFontInfo( const FName& InFontName, uint16 InSize, EFontHinting InHinting = EFontHinting::Default );
 
 	/**
-	 * Creates and initializes a new instance with the specified font name and size.
+	 * DEPRECATED - Creates and initializes a new instance with the specified font name and size.
 	 *
 	 * @param InFontName The name of the font.
 	 * @param InSize The size of the font.
@@ -72,7 +90,7 @@ public:
 	FSlateFontInfo( const ANSICHAR* InFontName, uint16 InSize, EFontHinting InHinting = EFontHinting::Default );
 
 	/**
-	 * Creates and initializes a new instance with the specified font name and size.
+	 * DEPRECATED - Creates and initializes a new instance with the specified font name and size.
 	 *
 	 * @param InFontName The name of the font.
 	 * @param InSize The size of the font.
@@ -90,7 +108,10 @@ public:
 	 */
 	bool operator==( const FSlateFontInfo& Other ) const 
 	{
-		return FontName == Other.FontName && Size == Other.Size && Hinting == Other.Hinting;
+		return FontObject == Other.FontObject
+			&& CompositeFont == Other.CompositeFont 
+			&& TypefaceFontName == Other.TypefaceFontName
+			&& Size == Other.Size;
 	}
 
 	/**
@@ -105,6 +126,16 @@ public:
 		return !(*this == Other);
 	}
 
+	/**
+	 * Check to see whether this font info has a valid composite font pointer set (either directly or via a UFont)
+	 */
+	bool HasValidFont() const;
+
+	/**
+	 * Get the composite font pointer associated with this font info (either directly or via a UFont)
+	 */
+	const FCompositeFont* GetCompositeFont() const;
+
 	/**
 	 * Calculates a type hash value for a font info.
 	 *
@@ -116,8 +147,32 @@ public:
 	friend inline uint32 GetTypeHash( const FSlateFontInfo& FontInfo )
 	{
 		uint32 Hash = 0;
-		Hash = HashCombine(Hash, GetTypeHash(FontInfo.FontName));
-		Hash = HashCombine(Hash, FontInfo.Size);
-		return HashCombine(Hash, uint32(FontInfo.Hinting));
+		Hash = HashCombine(Hash, GetTypeHash(FontInfo.FontObject));
+		Hash = HashCombine(Hash, GetTypeHash(FontInfo.CompositeFont));
+		Hash = HashCombine(Hash, GetTypeHash(FontInfo.TypefaceFontName));
+		Hash = HashCombine(Hash, GetTypeHash(FontInfo.Size));
+		return Hash;
 	}
+
+	/**
+	 * Used to upgrade legacy font into so that it uses composite fonts
+	 */
+	void PostSerialize(const FArchive& Ar);
+
+private:
+
+	/**
+	 * Used to upgrade legacy font into so that it uses composite fonts
+	 */
+	void UpgradeLegacyFontInfo();
+};
+
+template<>
+struct TStructOpsTypeTraits<FSlateFontInfo>
+	: public TStructOpsTypeTraitsBase
+{
+	enum 
+	{
+		WithPostSerialize = true,
+	};
 };
diff --git a/Engine/Source/Runtime/SlateCore/Public/Rendering/SlateRenderer.h b/Engine/Source/Runtime/SlateCore/Public/Rendering/SlateRenderer.h
--- a/Engine/Source/Runtime/SlateCore/Public/Rendering/SlateRenderer.h
+++ b/Engine/Source/Runtime/SlateCore/Public/Rendering/SlateRenderer.h
@@ -145,6 +145,11 @@ public:
 		return FontMeasure.ToSharedRef();
 	}
 
+	TSharedRef< class FSlateFontCache > GetFontCache() const 
+	{
+		return FontCache.ToSharedRef();
+	}
+
 	/**
 	 * Flushes all cached data from the font cache
 	 */
diff --git a/Engine/Source/Runtime/SlateCore/Public/Styling/SlateTypes.h b/Engine/Source/Runtime/SlateCore/Public/Styling/SlateTypes.h
--- a/Engine/Source/Runtime/SlateCore/Public/Styling/SlateTypes.h
+++ b/Engine/Source/Runtime/SlateCore/Public/Styling/SlateTypes.h
@@ -171,15 +171,18 @@ struct SLATECORE_API FTextBlockStyle : public FSlateWidgetStyle
 	UPROPERTY(EditAnywhere, BlueprintReadWrite, Category=Appearance)
 	FSlateFontInfo Font;
 	FTextBlockStyle& SetFont(const FSlateFontInfo& InFont) { Font = InFont; return *this; }
+	FTextBlockStyle& SetFont(TSharedPtr<const FCompositeFont> InCompositeFont, const int32 InSize, const FName& InTypefaceFontName = NAME_None) { Font = FSlateFontInfo(InCompositeFont, InSize, InTypefaceFontName); return *this; }
+	FTextBlockStyle& SetFont(const UObject* InFontObject, const int32 InSize, const FName& InTypefaceFontName = NAME_None) { Font = FSlateFontInfo(InFontObject, InSize, InTypefaceFontName); return *this; }
 	FTextBlockStyle& SetFont(const FName& InFontName, uint16 InSize) { Font = FSlateFontInfo(InFontName, InSize); return *this; }
 	FTextBlockStyle& SetFont(const FString& InFontName, uint16 InSize) { Font = FSlateFontInfo(*InFontName, InSize); return *this; }
 	FTextBlockStyle& SetFont(const WIDECHAR* InFontName, uint16 InSize) { Font = FSlateFontInfo(InFontName, InSize); return *this; }
 	FTextBlockStyle& SetFont(const ANSICHAR* InFontName, uint16 InSize) { Font = FSlateFontInfo(InFontName, InSize); return *this; }
-	FTextBlockStyle& SetFontName(const FName& InFontName) { Font.FontName = InFontName; return *this; }
-	FTextBlockStyle& SetFontName(const FString& InFontName) { Font.FontName = FName(*InFontName); return *this; }
-	FTextBlockStyle& SetFontName(const WIDECHAR* InFontName) { Font.FontName = InFontName; return *this; }
-	FTextBlockStyle& SetFontName(const ANSICHAR* InFontName) { Font.FontName = InFontName; return *this; }
+	FTextBlockStyle& SetFontName(const FName& InFontName) { Font = FSlateFontInfo(InFontName, Font.Size); return *this; }
+	FTextBlockStyle& SetFontName(const FString& InFontName) { Font = FSlateFontInfo(InFontName, Font.Size); return *this; }
+	FTextBlockStyle& SetFontName(const WIDECHAR* InFontName) { Font = FSlateFontInfo(InFontName, Font.Size); return *this; }
+	FTextBlockStyle& SetFontName(const ANSICHAR* InFontName) { Font = FSlateFontInfo(InFontName, Font.Size); return *this; }
 	FTextBlockStyle& SetFontSize(uint16 InSize) { Font.Size = InSize; return *this; }
+	FTextBlockStyle& SetTypefaceFontName(const FName& InTypefaceFontName) { Font.TypefaceFontName = InTypefaceFontName; return *this; }
 
 	/** The color and opacity of this text */
 	UPROPERTY(EditAnywhere, BlueprintReadWrite, Category=Appearance, meta=(DisplayName="Color"))
diff --git a/Engine/Source/Runtime/SlateRHIRenderer/Private/SlateRHIFontTexture.cpp b/Engine/Source/Runtime/SlateRHIRenderer/Private/SlateRHIFontTexture.cpp
--- a/Engine/Source/Runtime/SlateRHIRenderer/Private/SlateRHIFontTexture.cpp
+++ b/Engine/Source/Runtime/SlateRHIRenderer/Private/SlateRHIFontTexture.cpp
@@ -4,14 +4,58 @@
 
 
 FSlateFontTextureRHI::FSlateFontTextureRHI( uint32 InWidth, uint32 InHeight )
-	: FSlateTexture2DRHIRef( InWidth, InHeight, PF_A8, NULL, TexCreate_Dynamic, true )
-{ }
-
+	: Width( InWidth )
+	, Height( InHeight )
+{
+}
 
 void FSlateFontTextureRHI::InitDynamicRHI()
 {
+	check( IsInRenderingThread() );
+
 	// Create the texture
-	FSlateTexture2DRHIRef::InitDynamicRHI();
+	if( Width > 0 && Height > 0 )
+	{
+		check( !IsValidRef( ShaderResource) );
+		FRHIResourceCreateInfo CreateInfo;
+		ShaderResource = RHICreateTexture2D( Width, Height, PF_A8, 1, 1, TexCreate_Dynamic, CreateInfo );
+		check( IsValidRef( ShaderResource ) );
+
+		// Also assign the reference to the FTextureResource variable so that the Engine can access it
+		TextureRHI = ShaderResource;
+
+		const float MipMapBias = -1.0f;
+
+		// Create the sampler state RHI resource.
+		FSamplerStateInitializerRHI SamplerStateInitializer
+		(
+		  SF_Point,
+		  AM_Clamp,
+		  AM_Clamp,
+		  AM_Wrap,
+		  MipMapBias
+		);
+		SamplerStateRHI = RHICreateSamplerState(SamplerStateInitializer);
+
+		// Create a custom sampler state for using this texture in a deferred pass, where ddx / ddy are discontinuous
+		FSamplerStateInitializerRHI DeferredPassSamplerStateInitializer
+		(
+		  SF_Point,
+		  AM_Clamp,
+		  AM_Clamp,
+		  AM_Wrap,
+		  MipMapBias,
+		  // Disable anisotropic filtering, since aniso doesn't respect MaxLOD
+		  1,
+		  0,
+		  // Prevent the less detailed mip levels from being used, which hides artifacts on silhouettes due to ddx / ddy being very large
+		  // This has the side effect that it increases minification aliasing on light functions
+		  2
+		);
+		DeferredPassSamplerStateRHI = RHICreateSamplerState(DeferredPassSamplerStateInitializer);
+
+		INC_MEMORY_STAT_BY(STAT_SlateTextureGPUMemory, Width*Height*GPixelFormats[PF_A8].BlockBytes);
+	}
 
 	// Restore previous data if it exists
 	if( TempData.Num() > 0 )
@@ -28,6 +72,8 @@ void FSlateFontTextureRHI::InitDynamicRHI()
 
 void FSlateFontTextureRHI::ReleaseDynamicRHI()
 {
+	check( IsInRenderingThread() );
+
 	// Copy the data to temporary storage until InitDynamicRHI is called
 	uint32 Stride;
 	TempData.Empty();
@@ -38,7 +84,60 @@ void FSlateFontTextureRHI::ReleaseDynamicRHI()
 	RHIUnlockTexture2D( GetTypedResource(), 0, false );
 
 	// Release the texture
-	FSlateTexture2DRHIRef::ReleaseDynamicRHI();
+	if( IsValidRef(ShaderResource) )
+	{
+		DEC_MEMORY_STAT_BY(STAT_SlateTextureGPUMemory, Width*Height*GPixelFormats[PF_A8].BlockBytes);
+	}
+
+	ShaderResource.SafeRelease();
+}
+
+void FSlateFontTextureRHI::ResizeTexture( uint32 InWidth, uint32 InHeight )
+{
+	if (Width != InWidth || Height != InHeight)
+	{
+		if (IsInRenderingThread())
+		{
+			Width = InWidth;
+			Height = InHeight;
+			UpdateRHI();
+		}
+		else
+		{
+			FIntPoint Dimensions(InWidth, InHeight);
+			ENQUEUE_UNIQUE_RENDER_COMMAND_TWOPARAMETER(SlateResizeFontTexture,
+			FSlateFontTextureRHI*, TextureRHIRef, this,
+			FIntPoint, InDimensions, Dimensions,
+			{
+				TextureRHIRef->Width = InDimensions.X;
+				TextureRHIRef->Height = InDimensions.Y;
+				TextureRHIRef->UpdateRHI();
+			});
+		}
+	}
+}
+
+void FSlateFontTextureRHI::UpdateTexture( const TArray<uint8>& InBytes )
+{
+	if (IsInRenderingThread())
+	{
+		uint32 Stride = 0;
+		void* TextureBuffer = RHILockTexture2D(GetTypedResource(), 0, RLM_WriteOnly, Stride, false);
+		FMemory::Memcpy(TextureBuffer, InBytes.GetData(), InBytes.Num());
+		RHIUnlockTexture2D(GetTypedResource(), 0, false);
+	}
+	else
+	{
+		ENQUEUE_UNIQUE_RENDER_COMMAND_TWOPARAMETER(SlateUpdateFontTexture,
+		FSlateFontTextureRHI*, TextureRHIRef, this,
+		const TArray<uint8>&, TextureData, InBytes,
+		{
+			uint32 Stride = 0;
+			void* TextureBuffer = RHILockTexture2D(TextureRHIRef->GetTypedResource(), 0, RLM_WriteOnly, Stride, false);
+			FMemory::Memcpy(TextureBuffer, TextureData.GetData(), TextureData.Num());
+			RHIUnlockTexture2D(TextureRHIRef->GetTypedResource(), 0, false);
+		});
+	}
 }
 
 FSlateFontAtlasRHI::FSlateFontAtlasRHI( uint32 Width, uint32 Height )
@@ -48,38 +147,45 @@ FSlateFontAtlasRHI::FSlateFontAtlasRHI( uint32 Width, uint32 Height )
 }
 
 FSlateFontAtlasRHI::~FSlateFontAtlasRHI()
-{ 
-	if ( FontTexture )
-	{
-		delete FontTexture;
-	}
+{
 }
 
 void FSlateFontAtlasRHI::ReleaseResources()
 {
 	check( IsThreadSafeForSlateRendering() );
 
-	BeginReleaseResource( FontTexture );
+	BeginReleaseResource( FontTexture.Get() );
 }
 
-
 void FSlateFontAtlasRHI::ConditionalUpdateTexture()
 {
-	check( IsThreadSafeForSlateRendering() );
-
 	if( bNeedsUpdate )
 	{
-		BeginInitResource( FontTexture );
-
-		ENQUEUE_UNIQUE_RENDER_COMMAND_ONEPARAMETER( SlateUpdateFontTextureCommand,
-			FSlateFontAtlasRHI&, Atlas, *this,
+		if (IsInRenderingThread())
 		{
+			FontTexture->InitResource();
+
 			uint32 Stride;
-			uint8* TempData = (uint8*)RHILockTexture2D( Atlas.FontTexture->GetTypedResource(), 0, RLM_WriteOnly, Stride, false );
-			FMemory::Memcpy( TempData, Atlas.AtlasData.GetData(), Atlas.Stride*Atlas.AtlasWidth*Atlas.AtlasHeight );
-			RHIUnlockTexture2D( Atlas.FontTexture->GetTypedResource(),0,false );
-		});
+			uint8* TempData = (uint8*)RHILockTexture2D( FontTexture->GetTypedResource(), 0, RLM_WriteOnly, Stride, false );
+			FMemory::Memcpy( TempData, AtlasData.GetData(), Stride*AtlasWidth*AtlasHeight );
+			RHIUnlockTexture2D( FontTexture->GetTypedResource(),0,false );
+		}
+		else
+		{
+			check( IsThreadSafeForSlateRendering() );
+
+			BeginInitResource( FontTexture.Get() );
+
+			ENQUEUE_UNIQUE_RENDER_COMMAND_ONEPARAMETER( SlateUpdateFontTextureCommand,
+				FSlateFontAtlasRHI&, Atlas, *this,
+			{
+				uint32 Stride;
+				uint8* TempData = (uint8*)RHILockTexture2D( Atlas.FontTexture->GetTypedResource(), 0, RLM_WriteOnly, Stride, false );
+				FMemory::Memcpy( TempData, Atlas.AtlasData.GetData(), Atlas.Stride*Atlas.AtlasWidth*Atlas.AtlasHeight );
+				RHIUnlockTexture2D( Atlas.FontTexture->GetTypedResource(),0,false );
+			});
+		}
 
 		bNeedsUpdate = false;
 	}
-}
\ No newline at end of file
+}
diff --git a/Engine/Source/Runtime/SlateRHIRenderer/Private/SlateRHIFontTexture.h b/Engine/Source/Runtime/SlateRHIRenderer/Private/SlateRHIFontTexture.h
--- a/Engine/Source/Runtime/SlateRHIRenderer/Private/SlateRHIFontTexture.h
+++ b/Engine/Source/Runtime/SlateRHIRenderer/Private/SlateRHIFontTexture.h
@@ -8,7 +8,7 @@
  * ensure all characters in the font texture exist if the rendering resource has to be recreated 
  * between caching new characters
  */
-class FSlateFontTextureRHI : public FSlateTexture2DRHIRef
+class FSlateFontTextureRHI : public TSlateTexture<FTexture2DRHIRef>, public FSlateUpdatableTexture, public FTextureResource
 {
 public:
 	/** Constructor.  Initializes the texture
@@ -18,10 +18,29 @@ public:
 	 */
 	FSlateFontTextureRHI( uint32 InWidth, uint32 InHeight );
 
+	/** FSlateShaderResource interface */
+	virtual uint32 GetWidth() const override { return Width; }
+	virtual uint32 GetHeight() const override { return Height; }
+
+	/** FSlateUpdatableTexture interface */
+	virtual FSlateShaderResource* GetSlateResource() override { return this; }
+	virtual FRenderResource* GetRenderResource() override { return this; }
+	virtual void ResizeTexture( uint32 InWidth, uint32 InHeight ) override;
+	virtual void UpdateTexture( const TArray<uint8>& InBytes ) override;
+
+	/** FTextureResource interface */
+	virtual uint32 GetSizeX() const override { return Width; }
+	virtual uint32 GetSizeY() const override { return Height; }
+	virtual FString GetFriendlyName() const override { return TEXT("FSlateFontTextureRHI"); }
+
 	/** FRenderResource interface */
 	virtual void InitDynamicRHI() override;
 	virtual void ReleaseDynamicRHI() override;
 private:
+	/** Width of this texture */
+	uint32 Width;
+	/** Height of this texture */
+	uint32 Height;
 	/** Temporary data stored between Release and InitDynamicRHI */
 	TArray<uint8> TempData;
 };
@@ -38,10 +57,10 @@ public:
 	/**
 	 * FSlateFontAtlas interface 
 	 */
-	virtual class FSlateShaderResource* GetTexture()  override { return FontTexture; }
+	virtual class FSlateShaderResource* GetSlateTexture() override { return FontTexture.Get(); }
+	virtual class FTextureResource* GetEngineTexture() override { return FontTexture.Get(); }
 	virtual void ConditionalUpdateTexture()  override;
 	virtual void ReleaseResources() override;
 private:
-	FSlateTexture2DRHIRef* FontTexture;
-
+	TUniquePtr<FSlateFontTextureRHI> FontTexture;
 };
diff --git a/Engine/Source/Runtime/SlateRHIRenderer/Private/SlateRHIRenderingPolicy.cpp b/Engine/Source/Runtime/SlateRHIRenderer/Private/SlateRHIRenderingPolicy.cpp
--- a/Engine/Source/Runtime/SlateRHIRenderer/Private/SlateRHIRenderingPolicy.cpp
+++ b/Engine/Source/Runtime/SlateRHIRenderer/Private/SlateRHIRenderingPolicy.cpp
@@ -362,7 +362,7 @@ void FSlateRHIRenderingPolicy::DrawElements(FRHICommandListImmediate& RHICmdList
 				FTexture2DRHIRef TextureRHI;
 				if(ShaderResource)
 				{
-					TextureRHI = ShaderResource->GetType() == ESlateShaderResource::NativeTexture ? ((FSlateTexture2DRHIRef*)ShaderResource)->GetTypedResource() :  ((FSlateUTextureResource*)ShaderResource)->GetTypedResource();
+					TextureRHI = ((TSlateTexture<FTexture2DRHIRef>*)ShaderResource)->GetTypedResource();
 				}
 
 				if (ShaderType == ESlateShader::LineSegment)
diff --git a/Engine/Source/Runtime/UMG/Private/Components/EditableText.cpp b/Engine/Source/Runtime/UMG/Private/Components/EditableText.cpp
--- a/Engine/Source/Runtime/UMG/Private/Components/EditableText.cpp
+++ b/Engine/Source/Runtime/UMG/Private/Components/EditableText.cpp
@@ -38,16 +38,9 @@ void UEditableText::ReleaseSlateResources(bool bReleaseChildren)
 
 TSharedRef<SWidget> UEditableText::RebuildWidget()
 {
-	FString FontPath = FPaths::GameContentDir() / Font.FontName.ToString();
-
-	if ( !FPaths::FileExists(FontPath) )
-	{
-		FontPath = FPaths::EngineContentDir() / Font.FontName.ToString();
-	}
-	
 	MyEditableText = SNew(SEditableText)
 	.Style(&WidgetStyle)
-	.Font(FSlateFontInfo(FontPath, Font.Size))
+	.Font(Font)
 	.MinDesiredWidth(MinimumDesiredWidth)
 	.IsCaretMovedWhenGainFocus(IsCaretMovedWhenGainFocus)
 	.SelectAllTextWhenFocused(SelectAllTextWhenFocused)
diff --git a/Engine/Source/Runtime/UMG/Private/Components/EditableTextBox.cpp b/Engine/Source/Runtime/UMG/Private/Components/EditableTextBox.cpp
--- a/Engine/Source/Runtime/UMG/Private/Components/EditableTextBox.cpp
+++ b/Engine/Source/Runtime/UMG/Private/Components/EditableTextBox.cpp
@@ -41,16 +41,9 @@ void UEditableTextBox::ReleaseSlateResources(bool bReleaseChildren)
 
 TSharedRef<SWidget> UEditableTextBox::RebuildWidget()
 {
-	FString FontPath = FPaths::GameContentDir() / Font.FontName.ToString();
-
-	if ( !FPaths::FileExists(FontPath) )
-	{
-		FontPath = FPaths::EngineContentDir() / Font.FontName.ToString();
-	}
-
 	MyEditableTextBlock = SNew(SEditableTextBox)
 		.Style(&WidgetStyle)
-		.Font(FSlateFontInfo(FontPath, Font.Size))
+		.Font(Font)
 		.ForegroundColor(ForegroundColor)
 		.BackgroundColor(BackgroundColor)
 		.ReadOnlyForegroundColor(ReadOnlyForegroundColor)
diff --git a/Engine/Source/Runtime/UMG/Private/Components/MultiLineEditableText.cpp b/Engine/Source/Runtime/UMG/Private/Components/MultiLineEditableText.cpp
--- a/Engine/Source/Runtime/UMG/Private/Components/MultiLineEditableText.cpp
+++ b/Engine/Source/Runtime/UMG/Private/Components/MultiLineEditableText.cpp
@@ -27,16 +27,9 @@ void UMultiLineEditableText::ReleaseSlateResources(bool bReleaseChildren)
 
 TSharedRef<SWidget> UMultiLineEditableText::RebuildWidget()
 {
-	FString FontPath = FPaths::GameContentDir() / Font.FontName.ToString();
-
-	if ( !FPaths::FileExists(FontPath) )
-	{
-		FontPath = FPaths::EngineContentDir() / Font.FontName.ToString();
-	}
-	
 	MyMultiLineEditableText = SNew(SMultiLineEditableText)
 	.TextStyle(&WidgetStyle)
-	.Font(FSlateFontInfo(FontPath, Font.Size))
+	.Font(Font)
 	.Justification(Justification)
 	.WrapTextAt( WrapTextAt )
 	.AutoWrapText( bAutoWrapText )
diff --git a/Engine/Source/Runtime/UMG/Private/Components/MultiLineEditableTextBox.cpp b/Engine/Source/Runtime/UMG/Private/Components/MultiLineEditableTextBox.cpp
--- a/Engine/Source/Runtime/UMG/Private/Components/MultiLineEditableTextBox.cpp
+++ b/Engine/Source/Runtime/UMG/Private/Components/MultiLineEditableTextBox.cpp
@@ -33,17 +33,10 @@ void UMultiLineEditableTextBox::ReleaseSlateResources(bool bReleaseChildren)
 
 TSharedRef<SWidget> UMultiLineEditableTextBox::RebuildWidget()
 {
-	FString FontPath = FPaths::GameContentDir() / Font.FontName.ToString();
-
-	if ( !FPaths::FileExists(FontPath) )
-	{
-		FontPath = FPaths::EngineContentDir() / Font.FontName.ToString();
-	}
-	
 	MyEditableTextBlock = SNew(SMultiLineEditableTextBox)
 		.Style(&WidgetStyle)
 		.TextStyle(&TextStyle)
-		.Font(FSlateFontInfo(FontPath, Font.Size))
+		.Font(Font)
 		.Justification(Justification)
 		.ForegroundColor(ForegroundColor)
 		.BackgroundColor(BackgroundColor)
diff --git a/Engine/Source/Runtime/UMG/Private/Components/SpinBox.cpp b/Engine/Source/Runtime/UMG/Private/Components/SpinBox.cpp
--- a/Engine/Source/Runtime/UMG/Private/Components/SpinBox.cpp
+++ b/Engine/Source/Runtime/UMG/Private/Components/SpinBox.cpp
@@ -40,16 +40,9 @@ void USpinBox::ReleaseSlateResources(bool bReleaseChildren)
 
 TSharedRef<SWidget> USpinBox::RebuildWidget()
 {
-	FString FontPath = FPaths::GameContentDir() / Font.FontName.ToString();
-
-	if ( !FPaths::FileExists(FontPath) )
-	{
-		FontPath = FPaths::EngineContentDir() / Font.FontName.ToString();
-	}
-	
 	MySpinBox = SNew(SSpinBox<float>)
 	.Style(&WidgetStyle)
-	.Font(FSlateFontInfo(FontPath, Font.Size))
+	.Font(Font)
 	.ClearKeyboardFocusOnCommit(ClearKeyboardFocusOnCommit)
 	.SelectAllTextOnCommit(SelectAllTextOnCommit)
 	.OnValueChanged(BIND_UOBJECT_DELEGATE(FOnFloatValueChanged, HandleOnValueChanged))
diff --git a/Engine/Source/Runtime/UMG/Private/Components/TextBlock.cpp b/Engine/Source/Runtime/UMG/Private/Components/TextBlock.cpp
--- a/Engine/Source/Runtime/UMG/Private/Components/TextBlock.cpp
+++ b/Engine/Source/Runtime/UMG/Private/Components/TextBlock.cpp
@@ -70,19 +70,12 @@ void UTextBlock::SynchronizeProperties()
 {
 	Super::SynchronizeProperties();
 
-	FString FontPath = FPaths::GameContentDir() / Font.FontName.ToString();
-
-	if( !FPaths::FileExists(FontPath) )
-	{
-		FontPath = FPaths::EngineContentDir() / Font.FontName.ToString();
-	}
-
 	TAttribute<FText> TextBinding = OPTIONAL_BINDING(FText, Text);
 	TAttribute<FSlateColor> ColorAndOpacityBinding = OPTIONAL_BINDING(FSlateColor, ColorAndOpacity);
 	TAttribute<FLinearColor> ShadowColorAndOpacityBinding = OPTIONAL_BINDING(FLinearColor, ShadowColorAndOpacity);
 
 	MyTextBlock->SetText(TextBinding);
-	MyTextBlock->SetFont(FSlateFontInfo(FontPath, Font.Size));
+	MyTextBlock->SetFont(Font);
 	MyTextBlock->SetColorAndOpacity(ColorAndOpacityBinding);
 	MyTextBlock->SetShadowOffset(ShadowOffset);
 	MyTextBlock->SetShadowColorAndOpacity(ShadowColorAndOpacityBinding);
