Improved the composite font editor layout

It now shows the fonts within a font family in a grid rather than a list. This uses less space, but presents the same information.

The preview text was removed in favor of making the preview viewport handle displaying a preview of all of the fonts that make up the composite font.

diff --git a/Engine/Source/Editor/FontEditor/Private/SCompositeFontEditor.cpp b/Engine/Source/Editor/FontEditor/Private/SCompositeFontEditor.cpp
--- a/Engine/Source/Editor/FontEditor/Private/SCompositeFontEditor.cpp
+++ b/Engine/Source/Editor/FontEditor/Private/SCompositeFontEditor.cpp
@@ -346,11 +346,15 @@ void STypefaceEditor::Construct(const FArguments& InArgs)
 
 			+SVerticalBox::Slot()
 			.AutoHeight()
+			.Padding(FMargin(8.0f, 0.0f, 8.0f, 0.0f))
 			[
-				SAssignNew(TypefaceEntriesListView, SListView<FTypefaceListViewEntryPtr>)
+				SAssignNew(TypefaceEntriesTileView, STileView<FTypefaceListViewEntryPtr>)
 				.ListItemsSource(&TypefaceEntries)
 				.SelectionMode(ESelectionMode::None)
-				.OnGenerateRow(this, &STypefaceEditor::MakeTypefaceEntryWidget)
+				.ItemWidth(160)
+				.ItemHeight(120)
+				.ItemAlignment(EListItemAlignment::LeftAligned)
+				.OnGenerateTile(this, &STypefaceEditor::MakeTypefaceEntryWidget)
 			]
 		]
 	];
@@ -385,7 +389,7 @@ void STypefaceEditor::UpdateFontList()
 	// Add a dummy entry for the "Add" button slot
 	TypefaceEntries.Add(MakeShareable(new FTypefaceListViewEntry()));
 
-	TypefaceEntriesListView->RequestListRefresh();
+	TypefaceEntriesTileView->RequestListRefresh();
 }
 
 TSharedRef<ITableRow> STypefaceEditor::MakeTypefaceEntryWidget(FTypefaceListViewEntryPtr InTypefaceEntry, const TSharedRef<STableViewBase>& OwnerTable)
@@ -403,21 +407,22 @@ TSharedRef<ITableRow> STypefaceEditor::MakeTypefaceEntryWidget(FTypefaceListView
 			.ForegroundColor(FSlateColor::UseForeground())
 			.ToolTipText(LOCTEXT("AddFontTooltip", "Add a new font to this font family"))
 			.OnClicked(this, &STypefaceEditor::OnAddFont)
+			.VAlign(VAlign_Center)
 			[
-				SNew(SHorizontalBox)
+				SNew(SVerticalBox)
 
-				+SHorizontalBox::Slot()
-				.AutoWidth()
-				.Padding(8.0f)
-				.VAlign(VAlign_Center)
+				+SVerticalBox::Slot()
+				.AutoHeight()
+				.Padding(16.0f)
+				.HAlign(HAlign_Center)
 				[
 					SNew(SImage)
 					.Image(FEditorStyle::Get().GetBrush("FontEditor.Button_Add"))
 				]
 
-				+SHorizontalBox::Slot()
-				.AutoWidth()
-				.VAlign(VAlign_Center)
+				+SVerticalBox::Slot()
+				.AutoHeight()
+				.HAlign(HAlign_Center)
 				[
 					SNew(STextBlock)
 					.AutoWrapText(true)
@@ -441,7 +446,7 @@ TSharedRef<ITableRow> STypefaceEditor::MakeTypefaceEntryWidget(FTypefaceListView
 		SNew(STableRow<FTypefaceListViewEntryPtr>, OwnerTable)
 		[
 			SNew(SBox)
-			.Padding(FMargin(8.0f, 0.0f, 8.0f, 8.0f))
+			.Padding(FMargin(0.0f, 0.0f, 8.0f, 8.0f))
 			[
 				EntryWidget.ToSharedRef()
 			]
@@ -551,87 +556,80 @@ void STypefaceEntryEditor::Construct(const FArguments& InArgs)
 		.BorderImage(FEditorStyle::GetBrush("ToolPanel.GroupBorder"))
 		.Padding(8.0f)
 		[
-			SNew(SHorizontalBox)
+			SNew(SVerticalBox)
 
-			+SHorizontalBox::Slot()
+			+SVerticalBox::Slot()
+			.AutoHeight()
+			.Padding(FMargin(0.0f, 0.0f, 0.0f, 4.0f))
 			[
-				SNew(SVerticalBox)
-
-				+SVerticalBox::Slot()
-				.AutoHeight()
-				.Padding(FMargin(0.0f, 0.0f, 0.0f, 4.0f))
-				[
-					SNew(SHorizontalBox)
+				SNew(SInlineEditableTextBlock)
+				.Text(this, &STypefaceEntryEditor::GetTypefaceEntryName)
+				.ToolTipText(LOCTEXT("FontNameTooltip", "The name of this font within the font family (click to edit)"))
+				.OnTextCommitted(this, &STypefaceEntryEditor::OnTypefaceEntryNameCommitted)
+				.OnVerifyTextChanged(this, &STypefaceEntryEditor::OnTypefaceEntryChanged)
+			]
 
-					+SHorizontalBox::Slot()
-					[
-						SNew(SInlineEditableTextBlock)
-						.Text(this, &STypefaceEntryEditor::GetTypefaceEntryName)
-						.ToolTipText(LOCTEXT("FontNameTooltip", "The name of this font within the font family (click to edit)"))
-						.OnTextCommitted(this, &STypefaceEntryEditor::OnTypefaceEntryNameCommitted)
-						.OnVerifyTextChanged(this, &STypefaceEntryEditor::OnTypefaceEntryChanged)
-					]
+			+SVerticalBox::Slot()
+			.AutoHeight()
+			.Padding(FMargin(0.0f, 0.0f, 0.0f, 4.0f))
+			[
+				SNew(SHorizontalBox)
 
-					+SHorizontalBox::Slot()
-					.AutoWidth()
-					.Padding(FMargin(4.0f, 0.0f, 0.0f, 0.0f))
-					[
-						SNew(SBorder)
-						.BorderImage(FEditorStyle::GetBrush("ToolPanel.DarkGroupBorder"))
-						.VAlign(VAlign_Center)
-						.HAlign(HAlign_Center)
-						.Padding(2.0f)
-						[
-							SAssignNew(PreviewTextBlock, STextBlock)
-							.Text(FText::FromString(TEXT("Preview")))
-							.ToolTipText(LOCTEXT("PreviewFontTooltip", "Preview of how this font will look when rendered by Slate"))
-							.Font(this, &STypefaceEntryEditor::GetPreviewFontStyle)
-						]
-					]
+				+SHorizontalBox::Slot()
+				.VAlign(VAlign_Center)
+				[
+					SNew(STextBlock)
+					.Text(this, &STypefaceEntryEditor::GetTypefaceEntryFontLeafname)
+					.ToolTipText(this, &STypefaceEntryEditor::GetTypefaceEntryFontFilePath)
 				]
 
-				+SVerticalBox::Slot()
-				.AutoHeight()
+				+SHorizontalBox::Slot()
+				.AutoWidth()
+				.Padding(4.0f, 0.0f, 0.0f, 0.0f)
+				.VAlign(VAlign_Center)
 				[
-					SNew(SHorizontalBox)
-
-					+SHorizontalBox::Slot()
+					SNew(SButton)
+					.ButtonStyle(FEditorStyle::Get(), "HoverHintOnly")
+					.ToolTipText(LOCTEXT("FontFilePathPickerToolTip", "Choose a font file from this computer"))
+					.OnClicked(this, &STypefaceEntryEditor::OnBrowseTypefaceEntryFontPath)
+					.ContentPadding(2.0f)
+					.ForegroundColor(FSlateColor::UseForeground())
+					.IsFocusable(false)
 					[
-						SNew(SFilePathPicker)
-						.BrowseButtonImage(FEditorStyle::GetBrush("PropertyWindow.Button_Ellipsis"))
-						.BrowseButtonStyle(FEditorStyle::Get(), "HoverHintOnly")
-						.BrowseButtonToolTip(LOCTEXT("FontFilePathPickerToolTip", "Choose a font file from this computer"))
-						.BrowseDirectory_Static([]() { return FEditorDirectories::Get().GetLastDirectory(ELastDirectory::GENERIC_OPEN); })
-						.BrowseTitle(LOCTEXT("FontPickerTitle", "Choose a font file..."))
-						.FilePath(this, &STypefaceEntryEditor::GetTypefaceEntryFontFilePath)
-						.FileTypeFilter(TEXT("TrueType fonts (*.ttf)|*.ttf|OpenType fonts (*.otf)|*.otf"))
-						.OnPathPicked(this, &STypefaceEntryEditor::OnTypefaceEntryFontPathPicked)
+						SNew(SImage)
+						.Image(FEditorStyle::GetBrush("PropertyWindow.Button_Ellipsis"))
+						.ColorAndOpacity(FSlateColor::UseForeground())
 					]
+				]
+			]
 
-					+SHorizontalBox::Slot()
-					.AutoWidth()
-					.Padding(FMargin(4.0f, 0.0f, 0.0f, 0.0f))
+			+SVerticalBox::Slot()
+			.AutoHeight()
+			.Padding(FMargin(0.0f, 0.0f, 0.0f, 4.0f))
+			[
+				SNew(SBox)
+				.MinDesiredWidth(100.0f)
+				[
+					SNew(SComboBox<TSharedPtr<FFontHintingComboEntry>>)
+					.ToolTipText(LOCTEXT("HintingTooltip", "The hinting algorithm to use with this font"))
+					.OptionsSource(&HintingComboData)
+					.OnSelectionChanged(this, &STypefaceEntryEditor::OnHintingComboSelectionChanged)
+					.OnGenerateWidget(this, &STypefaceEntryEditor::MakeHintingComboEntryWidget)
 					[
-						SNew(SBox)
-						.MinDesiredWidth(100.0f)
-						[
-							SNew(SComboBox<TSharedPtr<FFontHintingComboEntry>>)
-							.OptionsSource(&HintingComboData)
-							.OnSelectionChanged(this, &STypefaceEntryEditor::OnHintingComboSelectionChanged)
-							.OnGenerateWidget(this, &STypefaceEntryEditor::MakeHintingComboEntryWidget)
-							[
-								SNew(STextBlock)
-								.Text(this, &STypefaceEntryEditor::GetHintingComboText)
-							]
-						]
+						SNew(STextBlock)
+						.Text(this, &STypefaceEntryEditor::GetHintingComboText)
 					]
 				]
 			]
 
-			+SHorizontalBox::Slot()
-			.AutoWidth()
-			.VAlign(VAlign_Center)
-			.Padding(FMargin(8.0f, 0.0f, 0.0f, 0.0f))
+			+SVerticalBox::Slot()
+			[
+				SNew(SSpacer)
+			]
+
+			+SVerticalBox::Slot()
+			.AutoHeight()
+			.HAlign(HAlign_Center)
 			[
 				SNew(SButton)
 				.ButtonStyle(FEditorStyle::Get(), "HoverHintOnly")
@@ -678,16 +676,61 @@ bool STypefaceEntryEditor::OnTypefaceEntryChanged(const FText& InNewName, FText&
 	return !OnVerifyFontName.IsBound() || OnVerifyFontName.Execute(TypefaceEntry, *InNewName.ToString(), OutFailureReason);
 }
 
-FString STypefaceEntryEditor::GetTypefaceEntryFontFilePath() const
+FText STypefaceEntryEditor::GetTypefaceEntryFontFilePath() const
+{
+	FTypefaceEntry* const TypefaceEntryPtr = TypefaceEntry->GetTypefaceEntry();
+
+	if(TypefaceEntryPtr)
+	{
+		return FText::FromString(TypefaceEntryPtr->Font.FontFilename);
+	}
+
+	return FText::GetEmpty();
+}
+
+FText STypefaceEntryEditor::GetTypefaceEntryFontLeafname() const
 {
 	FTypefaceEntry* const TypefaceEntryPtr = TypefaceEntry->GetTypefaceEntry();
 
 	if(TypefaceEntryPtr)
 	{
-		return TypefaceEntryPtr->Font.FontFilename;
+		return (TypefaceEntryPtr->Font.FontFilename.Len() > 0) 
+			? FText::FromString(FPaths::GetCleanFilename(TypefaceEntryPtr->Font.FontFilename)) 
+			: LOCTEXT("NoFontFileSelected", "<No File Selected>");
 	}
 
-	return FString();
+	return FText::GetEmpty();
+}
+
+FReply STypefaceEntryEditor::OnBrowseTypefaceEntryFontPath()
+{
+	IDesktopPlatform* const DesktopPlatform = FDesktopPlatformModule::Get();
+
+	if(DesktopPlatform)
+	{
+		const FString DefaultPath = FEditorDirectories::Get().GetLastDirectory(ELastDirectory::GENERIC_OPEN);
+
+		TSharedPtr<SWindow> ParentWindow = FSlateApplication::Get().FindWidgetWindow(AsShared());
+		const void* const ParentWindowHandle = (ParentWindow.IsValid() && ParentWindow->GetNativeWindow().IsValid())
+			? ParentWindow->GetNativeWindow()->GetOSWindowHandle()
+			: nullptr;
+
+		TArray<FString> OutFiles;
+		if(DesktopPlatform->OpenFileDialog(
+			ParentWindowHandle, 
+			LOCTEXT("FontPickerTitle", "Choose a font file...").ToString(), 
+			DefaultPath, 
+			TEXT(""), 
+			TEXT("TrueType fonts (*.ttf)|*.ttf|OpenType fonts (*.otf)|*.otf"), 
+			EFileDialogFlags::None, 
+			OutFiles
+			))
+		{
+			OnTypefaceEntryFontPathPicked(OutFiles[0]);
+		}
+	}
+
+	return FReply::Handled();
 }
 
 void STypefaceEntryEditor::OnTypefaceEntryFontPathPicked(const FString& InNewFontFilename)
@@ -854,7 +897,7 @@ void SSubTypefaceEditor::Construct(const FArguments& InArgs)
 				[
 					SNew(SButton)
 					.ButtonStyle(FEditorStyle::Get(), "HoverHintOnly")
-					.ToolTipText(LOCTEXT("DeleteFontTooltip", "Remove this sub-font family from the composite font"))
+					.ToolTipText(LOCTEXT("DeleteFontFamilyTooltip", "Remove this sub-font family from the composite font"))
 					.OnClicked(this, &SSubTypefaceEditor::OnDeleteSubFontFamilyClicked)
 					[
 						SNew(SImage)
diff --git a/Engine/Source/Editor/FontEditor/Private/SCompositeFontEditor.h b/Engine/Source/Editor/FontEditor/Private/SCompositeFontEditor.h
--- a/Engine/Source/Editor/FontEditor/Private/SCompositeFontEditor.h
+++ b/Engine/Source/Editor/FontEditor/Private/SCompositeFontEditor.h
@@ -135,8 +135,8 @@ private:
 	/** Internal list of font pointers for the list view (generated from TypefacePtr->Fonts) */
 	TArray<FTypefaceListViewEntryPtr> TypefaceEntries;
 
-	/** List view widget showing the font entries (uses TypefaceEntries as its source) */
-	TSharedPtr<SListView<FTypefaceListViewEntryPtr>> TypefaceEntriesListView;
+	/** Tile view widget showing the font entries (uses TypefaceEntries as its source) */
+	TSharedPtr<STileView<FTypefaceListViewEntryPtr>> TypefaceEntriesTileView;
 };
 
 class STypefaceEntryEditor : public SCompoundWidget
@@ -194,7 +194,13 @@ private:
 	bool OnTypefaceEntryChanged(const FText& InNewName, FText& OutFailureReason) const;
 
 	/** Get the current font filename of this typeface entry */
-	FString GetTypefaceEntryFontFilePath() const;
+	FText GetTypefaceEntryFontFilePath() const;
+
+	/** Get the current font leafname of this typeface entry */
+	FText GetTypefaceEntryFontLeafname() const;
+
+	/** Open a file picker to let you pick a new font file */
+	FReply OnBrowseTypefaceEntryFontPath();
 
 	/** Set the current font filename (and associated font data) */
 	void OnTypefaceEntryFontPathPicked(const FString& InNewFontFilename);
diff --git a/Engine/Source/Editor/FontEditor/Private/SFontEditorViewport.cpp b/Engine/Source/Editor/FontEditor/Private/SFontEditorViewport.cpp
--- a/Engine/Source/Editor/FontEditor/Private/SFontEditorViewport.cpp
+++ b/Engine/Source/Editor/FontEditor/Private/SFontEditorViewport.cpp
@@ -169,9 +169,60 @@ void FFontEditorViewportClient::Draw(FViewport* Viewport, FCanvas* Canvas)
 		// Erase with our background color
 		Canvas->Clear( BackgroundColor);
 
+		static const FVector2D StartPos(4.0f, 4.0f);
+
 		// And draw the text with the foreground color
-		FCanvasTextItem TextItem( FVector2D::ZeroVector, PreviewText, Font, FLinearColor(ForegroundColor) );
-		Canvas->DrawItem( TextItem );		
+		if (Font->FontCacheType == EFontCacheType::Runtime && Font->CompositeFont.DefaultTypeface.Fonts.Num() > 0)
+		{
+			static const float FontScale = 1.0f;
+
+			TSharedRef<FSlateFontCache> FontCache = FSlateApplication::Get().GetRenderer()->GetFontCache();
+			TSharedRef<FSlateFontMeasure> FontMeasure = FSlateApplication::Get().GetRenderer()->GetFontMeasureService();
+
+			FVector2D CurPos = StartPos;
+			float WidestName = 0.0f;
+
+			// Draw and measure each name so we can work out where to start drawing the preview text column
+			for (const FTypefaceEntry& TypefaceEntry : Font->CompositeFont.DefaultTypeface.Fonts)
+			{
+				const FSlateFontInfo FontInfo(Font, Font->LegacyFontSize, TypefaceEntry.Name);
+
+				FCharacterList& CharacterList = FontCache->GetCharacterList(FontInfo, FontScale);
+				const int32 MaxCharHeight = CharacterList.GetMaxHeight();
+
+				const FText EntryNameText = FText::FromName(TypefaceEntry.Name);
+
+				FCanvasTextItem TextItem(CurPos, EntryNameText, FontInfo, FLinearColor(ForegroundColor));
+				Canvas->DrawItem(TextItem);
+
+				const FVector2D MeasuredText = FontMeasure->Measure(EntryNameText, FontInfo, FontScale);
+				WidestName = FMath::Max(WidestName, MeasuredText.X);
+
+				CurPos.Y += MaxCharHeight;
+			}
+
+			CurPos = FVector2D(WidestName + 12.0f, StartPos.Y);
+
+			// Draw the preview text using each of the default fonts
+			for (const FTypefaceEntry& TypefaceEntry : Font->CompositeFont.DefaultTypeface.Fonts)
+			{
+				const FSlateFontInfo FontInfo(Font, Font->LegacyFontSize, TypefaceEntry.Name);
+
+				FCharacterList& CharacterList = FontCache->GetCharacterList(FontInfo, FontScale);
+				const int32 MaxCharHeight = CharacterList.GetMaxHeight();
+
+				FCanvasTextItem TextItem(CurPos, PreviewText, FontInfo, FLinearColor(ForegroundColor));
+				Canvas->DrawItem(TextItem);
+
+				CurPos.Y += MaxCharHeight;
+			}
+		}
+		else
+		{
+			FCanvasTextItem TextItem(StartPos, PreviewText, Font, FLinearColor(ForegroundColor));
+			TextItem.BlendMode = (Font->ImportOptions.bUseDistanceFieldAlpha) ? SE_BLEND_TranslucentDistanceField : SE_BLEND_Translucent;
+			Canvas->DrawItem(TextItem);
+		}
 	}
 }
 
