The font editor now hides the properties that don't belong to the active font cache type

diff --git a/Engine/Source/Editor/FontEditor/Private/FontEditor.cpp b/Engine/Source/Editor/FontEditor/Private/FontEditor.cpp
--- a/Engine/Source/Editor/FontEditor/Private/FontEditor.cpp
+++ b/Engine/Source/Editor/FontEditor/Private/FontEditor.cpp
@@ -530,6 +530,7 @@ void FFontEditor::CreateInternalWidgets()
 	FontProperties = PropertyModule.CreateDetailView(Args);
 	FontPageProperties = PropertyModule.CreateDetailView(Args);
 
+	FontProperties->SetIsPropertyVisibleDelegate(FIsPropertyVisible::CreateRaw(this, &FFontEditor::GetIsPropertyVisible));
 	FontProperties->SetObject( Font );
 }
 
@@ -993,6 +994,18 @@ bool FFontEditor::RecreateFontObject(const EFontCacheType NewCacheType)
 	return bSuccess;
 }
 
+bool FFontEditor::GetIsPropertyVisible(const FPropertyAndParent& PropertyAndParent) const
+{
+	static const FName CategoryFName = "Category";
+
+	// We need to hide the properties associated with the category that we're not currently using (either Offline or Runtime)
+	const FString CategoryToExclude = (Font->FontCacheType == EFontCacheType::Offline) ? TEXT("RuntimeFont") : TEXT("OfflineFont");
+
+	// We need to hide the properties associated with the category that we're not currently using (either Offline or Runtime)
+	const FString& CategoryValue = PropertyAndParent.Property.GetMetaData(CategoryFName);
+	return CategoryValue != CategoryToExclude;
+}
+
 bool FFontEditor::ShouldPromptForNewFilesOnReload(const UObject& EditingObject) const
 {
 	return false;
diff --git a/Engine/Source/Editor/FontEditor/Private/FontEditor.h b/Engine/Source/Editor/FontEditor/Private/FontEditor.h
--- a/Engine/Source/Editor/FontEditor/Private/FontEditor.h
+++ b/Engine/Source/Editor/FontEditor/Private/FontEditor.h
@@ -109,6 +109,9 @@ private:
 	/** Recreate the font object so that it's using the given caching method */
 	bool RecreateFontObject(const EFontCacheType NewCacheType);
 
+	/** Check to see if the given property should be visible in the details panel */
+	bool GetIsPropertyVisible(const struct FPropertyAndParent& PropertyAndParent) const;
+
 private:
 	/** The font asset being inspected */
 	UFont* Font;
