CIS fix.

Missing header in non-unity build.

diff --git a/Engine/Source/Editor/FontEditor/Private/SCompositeFontEditor.cpp b/Engine/Source/Editor/FontEditor/Private/SCompositeFontEditor.cpp
--- a/Engine/Source/Editor/FontEditor/Private/SCompositeFontEditor.cpp
+++ b/Engine/Source/Editor/FontEditor/Private/SCompositeFontEditor.cpp
@@ -6,6 +6,7 @@
 #include "SInlineEditableTextBlock.h"
 #include "SNumericEntryBox.h"
 #include "ScopedTransaction.h"
+#include "DesktopPlatformModule.h"
 
 #define LOCTEXT_NAMESPACE "FontEditor"
 
