Added support for setting a friendly, editor-only, name for sub-font families

These let you specify the purpose of the font family, rather than them all being called "Sub-Font Family #X".

diff --git a/Engine/Source/Editor/FontEditor/Private/SCompositeFontEditor.cpp b/Engine/Source/Editor/FontEditor/Private/SCompositeFontEditor.cpp
--- a/Engine/Source/Editor/FontEditor/Private/SCompositeFontEditor.cpp
+++ b/Engine/Source/Editor/FontEditor/Private/SCompositeFontEditor.cpp
@@ -251,15 +251,11 @@ TSharedRef<ITableRow> SCompositeFontEditor::MakeSubTypefaceEntryWidget(FSubTypef
 	}
 	else
 	{
-		const int32 SubFontIndex = SubTypefaceEntries.IndexOfByKey(InSubTypefaceEntry);
-		const FText SubFontFamilyName = FText::Format(LOCTEXT("SubFontFamilyNameFmt", "Sub-Font Family #{0}"), FText::AsNumber(SubFontIndex + 1));
-
 		EntryWidget = SNew(SSubTypefaceEditor)
 		.CompositeFontEditor(this)
 		.SubTypeface(InSubTypefaceEntry)
 		.ParentTypeface(this, &SCompositeFontEditor::GetConstDefaultTypeface)
-		.OnDeleteSubFontFamily(this, &SCompositeFontEditor::OnDeleteSubFontFamily)
-		.TypefaceDisplayName(SubFontFamilyName);
+		.OnDeleteSubFontFamily(this, &SCompositeFontEditor::OnDeleteSubFontFamily);
 	}
 
 	return
@@ -327,9 +323,12 @@ void STypefaceEditor::Construct(const FArguments& InArgs)
 
 				+SHorizontalBox::Slot()
 				[
-					SNew(STextBlock)
+					SNew(SInlineEditableTextBlock)
 					.Text(InArgs._TypefaceDisplayName)
+					.ToolTipText((InArgs._OnDisplayNameCommitted.IsBound()) ? LOCTEXT("FontFamilyNameTooltip", "The name of this font family (click to edit)") : FText::GetEmpty())
 					.Font(FEditorStyle::GetFontStyle("DetailsView.CategoryFontStyle"))
+					.OnTextCommitted(InArgs._OnDisplayNameCommitted)
+					.IsReadOnly(!InArgs._OnDisplayNameCommitted.IsBound())
 				]
 
 				+SHorizontalBox::Slot()
@@ -848,7 +847,8 @@ void SSubTypefaceEditor::Construct(const FArguments& InArgs)
 		SAssignNew(TypefaceEditor, STypefaceEditor)
 		.CompositeFontEditor(InArgs._CompositeFontEditor)
 		.Typeface(this, &SSubTypefaceEditor::GetTypeface)
-		.TypefaceDisplayName(InArgs._TypefaceDisplayName)
+		.TypefaceDisplayName(this, &SSubTypefaceEditor::GetDisplayName)
+		.OnDisplayNameCommitted(this, &SSubTypefaceEditor::OnDisplayNameCommitted)
 		.HeaderContent()
 		[
 			SNew(SBox)
@@ -932,6 +932,41 @@ FTypeface* SSubTypefaceEditor::GetTypeface() const
 	return (SubTypefaceEntryPtr) ? &SubTypefaceEntryPtr->Typeface : nullptr;
 }
 
+FText SSubTypefaceEditor::GetDisplayName() const
+{
+	const FCompositeSubFont* const SubTypefaceEntryPtr = SubTypeface->GetSubTypefaceEntry();
+	
+	if(SubTypefaceEntryPtr)
+	{
+		return (SubTypefaceEntryPtr->EditorName.IsNone()) 
+			? FText::Format(LOCTEXT("SubFontFamilyNameFmt", "Sub-Font Family #{0}"), FText::AsNumber(SubTypeface->SubTypefaceEntryIndex + 1))
+			: FText::FromName(SubTypefaceEntryPtr->EditorName);
+	}
+
+	return FText::GetEmpty();
+}
+
+void SSubTypefaceEditor::OnDisplayNameCommitted(const FText& InNewName, ETextCommit::Type InCommitType)
+{
+	FCompositeSubFont* const SubTypefaceEntryPtr = SubTypeface->GetSubTypefaceEntry();
+	
+	if(SubTypefaceEntryPtr)
+	{
+		const FScopedTransaction Transaction(LOCTEXT("SetFontFamilyDisplayName", "Set Font Family Display Name"));
+		CompositeFontEditorPtr->GetFontObject()->Modify();
+
+		const FText DefaultText = FText::Format(LOCTEXT("SubFontFamilyNameFmt", "Sub-Font Family #{0}"), FText::AsNumber(SubTypeface->SubTypefaceEntryIndex + 1));
+		if(InNewName.ToString().Equals(DefaultText.ToString()))
+		{
+			SubTypefaceEntryPtr->EditorName = NAME_None;
+		}
+		else
+		{
+			SubTypefaceEntryPtr->EditorName = *InNewName.ToString();
+		}
+	}
+}
+
 EVisibility SSubTypefaceEditor::GetAddFontOverrideVisibility() const
 {
 	return (ParentTypeface.Get(nullptr)) ? EVisibility::Visible : EVisibility::Collapsed;
diff --git a/Engine/Source/Editor/FontEditor/Private/SCompositeFontEditor.h b/Engine/Source/Editor/FontEditor/Private/SCompositeFontEditor.h
--- a/Engine/Source/Editor/FontEditor/Private/SCompositeFontEditor.h
+++ b/Engine/Source/Editor/FontEditor/Private/SCompositeFontEditor.h
@@ -80,6 +80,7 @@ public:
 		: _CompositeFontEditor(nullptr)
 		, _Typeface(nullptr)
 		, _TypefaceDisplayName()
+		, _OnDisplayNameCommitted()
 		, _HeaderContent()
 		, _BodyContent()
 		{}
@@ -93,6 +94,9 @@ public:
 		/** Name to show in the UI for this typeface */
 		SLATE_ATTRIBUTE(FText, TypefaceDisplayName)
 
+		/** Callback to use when the display name is committed (if not set, the display name will be read-only) */
+		SLATE_EVENT(FOnTextCommitted, OnDisplayNameCommitted)
+
 		/** Slot for extra content to place in the header bar (optional) */
 		SLATE_NAMED_SLOT(FArguments, HeaderContent)
 
@@ -255,7 +259,6 @@ public:
 		: _CompositeFontEditor(nullptr)
 		, _SubTypeface()
 		, _ParentTypeface(nullptr)
-		, _TypefaceDisplayName()
 		, _OnDeleteSubFontFamily()
 		{}
 
@@ -268,9 +271,6 @@ public:
 		/** Parent typeface to inherit font slots from (may be invalid, or change in response to an undo/redo) */
 		SLATE_ATTRIBUTE(const FTypeface*, ParentTypeface)
 
-		/** Name to show in the UI for this typeface */
-		SLATE_ATTRIBUTE(FText, TypefaceDisplayName)
-
 		/** Called when this typeface entry should be deleted */
 		SLATE_EVENT(FOnDeleteSubFontFamily, OnDeleteSubFontFamily)
 
@@ -286,6 +286,12 @@ private:
 	/** Get the typeface used by this sub-typeface */
 	FTypeface* GetTypeface() const;
 
+	/** Get the display name of this sub-font family */
+	FText GetDisplayName() const;
+
+	/** Set the display name of this sub-font family */
+	void OnDisplayNameCommitted(const FText& InNewName, ETextCommit::Type InCommitType);
+
 	/** Get the visibility of the "Add Font Override" combo button */
 	EVisibility GetAddFontOverrideVisibility() const;
 
diff --git a/Engine/Source/Runtime/SlateCore/Public/Fonts/CompositeFont.h b/Engine/Source/Runtime/SlateCore/Public/Fonts/CompositeFont.h
--- a/Engine/Source/Runtime/SlateCore/Public/Fonts/CompositeFont.h
+++ b/Engine/Source/Runtime/SlateCore/Public/Fonts/CompositeFont.h
@@ -138,6 +138,12 @@ struct SLATECORE_API FCompositeSubFont
 	/** Amount to scale this sub-font so that it better matches the size of the default font */
 	UPROPERTY(EditAnywhere, BlueprintReadWrite, Category=Font)
 	float ScalingFactor;
+
+#if WITH_EDITORONLY_DATA
+	/** Name of this sub-font. Only used by the editor UI as a convenience to let you state the purpose of the font family. */
+	UPROPERTY()
+	FName EditorName;
+#endif
 };
 
 USTRUCT(BlueprintType)
