Fixed crash when attempting to generate too large a font texture.

diff --git a/Engine/Source/Editor/FontEditor/Private/FontEditor.cpp b/Engine/Source/Editor/FontEditor/Private/FontEditor.cpp
--- a/Engine/Source/Editor/FontEditor/Private/FontEditor.cpp
+++ b/Engine/Source/Editor/FontEditor/Private/FontEditor.cpp
@@ -385,8 +385,11 @@ void FFontEditor::PostUndo(bool bSuccess)
 
 void FFontEditor::NotifyPostChange( const FPropertyChangedEvent& PropertyChangedEvent, class FEditPropertyChain* PropertyThatChanged)
 {
-	static const FName FontCacheTypePropertyName("FontCacheType");
-	static const FName CompositeFontPropertyName("CompositeFont");
+	static const FName FontCacheTypePropertyName = GET_MEMBER_NAME_CHECKED(UFont, FontCacheType);
+	static const FName CompositeFontPropertyName = GET_MEMBER_NAME_CHECKED(UFont, CompositeFont);
+	static const FName TexturePageWidthName = GET_MEMBER_NAME_CHECKED(FFontImportOptionsData, TexturePageWidth);
+	static const FName TexturePageMaxHeightName = GET_MEMBER_NAME_CHECKED(FFontImportOptionsData, TexturePageMaxHeight);
+	static const FName DistanceFieldScaleFactorName = GET_MEMBER_NAME_CHECKED(FFontImportOptionsData, DistanceFieldScaleFactor);
 
 	if(PropertyChangedEvent.Property && PropertyChangedEvent.Property->GetFName() == FontCacheTypePropertyName)
 	{
@@ -429,6 +432,48 @@ void FFontEditor::NotifyPostChange( const FPropertyChangedEvent& PropertyChanged
 		}
 	}
 
+	if (PropertyChangedEvent.Property && PropertyChangedEvent.Property->GetFName() == DistanceFieldScaleFactorName)
+	{
+		const uint32 SignedInt32NumBits = 31;
+		const uint32 Log2TexturePageWidth = FMath::CeilLogTwo(Font->ImportOptions.TexturePageWidth);
+		const uint32 Log2TexturePageMaxHeight = FMath::CeilLogTwo(Font->ImportOptions.TexturePageMaxHeight);
+		const uint32 Log2BytesPerPixel = 2;
+
+		const int32 MaxDistanceFieldScaleFactor = 1 << ((SignedInt32NumBits - Log2BytesPerPixel - Log2TexturePageWidth - Log2TexturePageMaxHeight) / 2);
+		if (Font->ImportOptions.DistanceFieldScaleFactor > MaxDistanceFieldScaleFactor)
+		{
+			Font->ImportOptions.DistanceFieldScaleFactor = MaxDistanceFieldScaleFactor;
+		}
+	}
+
+	if (PropertyChangedEvent.Property && PropertyChangedEvent.Property->GetFName() == TexturePageWidthName)
+	{
+		const uint32 SignedInt32NumBits = 31;
+		const uint32 Log2DistanceFieldScaleFactor = FMath::Max(1U, FMath::CeilLogTwo(Font->ImportOptions.DistanceFieldScaleFactor));
+		const uint32 Log2TexturePageMaxHeight = FMath::CeilLogTwo(Font->ImportOptions.TexturePageMaxHeight);
+		const uint32 Log2BytesPerPixel = 2;
+
+		const int32 MaxTexturePageWidth = 1 << (SignedInt32NumBits - Log2BytesPerPixel - 2 * Log2DistanceFieldScaleFactor - Log2TexturePageMaxHeight);
+		if (Font->ImportOptions.TexturePageWidth > MaxTexturePageWidth)
+		{
+			Font->ImportOptions.TexturePageWidth = MaxTexturePageWidth;
+		}
+	}
+
+	if (PropertyChangedEvent.Property && PropertyChangedEvent.Property->GetFName() == TexturePageMaxHeightName)
+	{
+		const uint32 SignedInt32NumBits = 31;
+		const uint32 Log2DistanceFieldScaleFactor = FMath::Max(1U, FMath::CeilLogTwo(Font->ImportOptions.DistanceFieldScaleFactor));
+		const uint32 Log2TexturePageWidth = FMath::CeilLogTwo(Font->ImportOptions.TexturePageWidth);
+		const uint32 Log2BytesPerPixel = 2;
+
+		const int32 MaxTexturePageMaxHeight = 1 << (SignedInt32NumBits - Log2BytesPerPixel - 2 * Log2DistanceFieldScaleFactor - Log2TexturePageWidth);
+		if (Font->ImportOptions.TexturePageMaxHeight > MaxTexturePageMaxHeight)
+		{
+			Font->ImportOptions.TexturePageMaxHeight = MaxTexturePageMaxHeight;
+		}
+	}
+
 	// If we changed a property of the composite font, we need to refresh the composite font editor
 	if(PropertyThatChanged && PropertyThatChanged->GetHead()->GetValue()->GetFName() == CompositeFontPropertyName)
 	{
