Fixed wrong preview text for font previews

diff --git a/Engine/Source/Editor/FontEditor/Private/FontEditor.cpp b/Engine/Source/Editor/FontEditor/Private/FontEditor.cpp
--- a/Engine/Source/Editor/FontEditor/Private/FontEditor.cpp
+++ b/Engine/Source/Editor/FontEditor/Private/FontEditor.cpp
@@ -563,7 +563,7 @@ void FFontEditor::CreateInternalWidgets()
 	.AutoHeight()
 	[
 		SAssignNew(FontPreviewText, SEditableTextBox)
-		.Text(LOCTEXT("DefaultPreviewText", "The quick brown fox jumped over the lazy dog"))
+		.Text(LOCTEXT("DefaultPreviewText", "The quick brown fox jumps over the lazy dog"))
 		.SelectAllTextWhenFocused(true)
 		.OnTextChanged(this, &FFontEditor::OnPreviewTextChanged)
 	];
diff --git a/Engine/Source/Editor/FontEditor/Private/SFontEditorViewport.cpp b/Engine/Source/Editor/FontEditor/Private/SFontEditorViewport.cpp
--- a/Engine/Source/Editor/FontEditor/Private/SFontEditorViewport.cpp
+++ b/Engine/Source/Editor/FontEditor/Private/SFontEditorViewport.cpp
@@ -74,7 +74,7 @@ FFontEditorViewportClient::FFontEditorViewportClient(TWeakPtr<SFontEditorViewpor
 	, CurrentSelectedPage(INDEX_NONE)
 	, PageGap(4)
 {
-	PreviewText = LOCTEXT("DefaultPreviewText", "The quick brown fox jumped over the lazy dog");
+	PreviewText = LOCTEXT("DefaultPreviewText", "The quick brown fox jumps over the lazy dog");
 	BackgroundColor = FColor::Black;
 	ForegroundColor = FColor::White;
 }
