The composite font editor import dialog now shows you all supported font files by default

diff --git a/Engine/Source/Editor/FontEditor/Private/SCompositeFontEditor.cpp b/Engine/Source/Editor/FontEditor/Private/SCompositeFontEditor.cpp
--- a/Engine/Source/Editor/FontEditor/Private/SCompositeFontEditor.cpp
+++ b/Engine/Source/Editor/FontEditor/Private/SCompositeFontEditor.cpp
@@ -722,7 +722,7 @@ FReply STypefaceEntryEditor::OnBrowseTypefaceEntryFontPath()
 			LOCTEXT("FontPickerTitle", "Choose a font file...").ToString(), 
 			DefaultPath, 
 			TEXT(""), 
-			TEXT("TrueType fonts (*.ttf)|*.ttf|OpenType fonts (*.otf)|*.otf"), 
+			TEXT("All Font Files (*.ttf, *.otf)|*.ttf;*.otf|TrueType fonts (*.ttf)|*.ttf|OpenType fonts (*.otf)|*.otf"), 
 			EFileDialogFlags::None, 
 			OutFiles
 			))
