Added some tooltips to the composite font editor

UE-17429 - Sub-Font Family UI in the font editor needs tool tips

diff --git a/Engine/Source/Editor/FontEditor/Private/SCompositeFontEditor.cpp b/Engine/Source/Editor/FontEditor/Private/SCompositeFontEditor.cpp
--- a/Engine/Source/Editor/FontEditor/Private/SCompositeFontEditor.cpp
+++ b/Engine/Source/Editor/FontEditor/Private/SCompositeFontEditor.cpp
@@ -1227,6 +1227,7 @@ void SCharacterRangeEditor::Construct(const FArguments& InArgs)
 			SNew(SEditableTextBox)
 			.Text(this, &SCharacterRangeEditor::GetRangeComponentAsTCHAR, 0)
 			.OnTextCommitted(this, &SCharacterRangeEditor::OnRangeComponentCommittedAsTCHAR, 0)
+			.ToolTipText(LOCTEXT("MinCharacterRangeEditCharTooltip", "Specifies the lower inclusive boundary of this character range as a character value"))
 		]
 
 		+SGridPanel::Slot(0, 1)
@@ -1235,6 +1236,7 @@ void SCharacterRangeEditor::Construct(const FArguments& InArgs)
 			SNew(SEditableTextBox)
 			.Text(this, &SCharacterRangeEditor::GetRangeComponentAsHexString, 0)
 			.OnTextCommitted(this, &SCharacterRangeEditor::OnRangeComponentCommittedAsHexString, 0)
+			.ToolTipText(LOCTEXT("MinCharacterRangeEditHexTooltip", "Specifies the lower inclusive boundary of this character range as a hexadecimal value"))
 		]
 
 		+SGridPanel::Slot(0, 2)
@@ -1243,6 +1245,7 @@ void SCharacterRangeEditor::Construct(const FArguments& InArgs)
 			SNew(SNumericEntryBox<int32>)
 			.Value(this, &SCharacterRangeEditor::GetRangeComponentAsOptional, 0)
 			.OnValueCommitted(this, &SCharacterRangeEditor::OnRangeComponentCommittedAsNumeric, 0)
+			.ToolTipText(LOCTEXT("MinCharacterRangeEditDecTooltip", "Specifies the lower inclusive boundary of this character range as a decimal value"))
 		]
 
 		// Separator
@@ -1262,6 +1265,7 @@ void SCharacterRangeEditor::Construct(const FArguments& InArgs)
 			SNew(SEditableTextBox)
 			.Text(this, &SCharacterRangeEditor::GetRangeComponentAsTCHAR, 1)
 			.OnTextCommitted(this, &SCharacterRangeEditor::OnRangeComponentCommittedAsTCHAR, 1)
+			.ToolTipText(LOCTEXT("MaxCharacterRangeEditCharTooltip", "Specifies the upper inclusive boundary of this character range as a character value"))
 		]
 
 		+SGridPanel::Slot(2, 1)
@@ -1270,6 +1274,7 @@ void SCharacterRangeEditor::Construct(const FArguments& InArgs)
 			SNew(SEditableTextBox)
 			.Text(this, &SCharacterRangeEditor::GetRangeComponentAsHexString, 1)
 			.OnTextCommitted(this, &SCharacterRangeEditor::OnRangeComponentCommittedAsHexString, 1)
+			.ToolTipText(LOCTEXT("MaxCharacterRangeEditHexTooltip", "Specifies the upper inclusive boundary of this character range as a hexadecimal value"))
 		]
 
 		+SGridPanel::Slot(2, 2)
@@ -1278,6 +1283,7 @@ void SCharacterRangeEditor::Construct(const FArguments& InArgs)
 			SNew(SNumericEntryBox<int32>)
 			.Value(this, &SCharacterRangeEditor::GetRangeComponentAsOptional, 1)
 			.OnValueCommitted(this, &SCharacterRangeEditor::OnRangeComponentCommittedAsNumeric, 1)
+			.ToolTipText(LOCTEXT("MaxCharacterRangeEditDecTooltip", "Specifies the upper inclusive boundary of this character range as a decimal value"))
 		]
 	];
 }
