Improved the UX in the composite font editor

UE-18107 - Broken UX in Font Editor

Adding a new font or sub-font family will now automatically open their names for editing to teach you that they can be renamed.

This also changed the icon for opening a font file as there was some confusion about the ellipsis being for the text that was cut off, rather than being a button to pick a font file.

diff --git a/Engine/Source/Editor/FontEditor/Private/SCompositeFontEditor.cpp b/Engine/Source/Editor/FontEditor/Private/SCompositeFontEditor.cpp
--- a/Engine/Source/Editor/FontEditor/Private/SCompositeFontEditor.cpp
+++ b/Engine/Source/Editor/FontEditor/Private/SCompositeFontEditor.cpp
@@ -16,16 +16,19 @@ struct FTypefaceListViewEntry
 {
 	TAttribute<FTypeface*> Typeface;
 	int32 TypefaceEntryIndex;
+	bool bRenameRequested;
 
 	FTypefaceListViewEntry()
 		: Typeface()
 		, TypefaceEntryIndex(INDEX_NONE)
+		, bRenameRequested(false)
 	{
 	}
 
 	FTypefaceListViewEntry(TAttribute<FTypeface*>& InTypeface, const int32 InTypefaceEntryIndex)
 		: Typeface(InTypeface)
 		, TypefaceEntryIndex(InTypefaceEntryIndex)
+		, bRenameRequested(false)
 	{
 	}
 
@@ -46,16 +49,19 @@ struct FSubTypefaceListViewEntry
 {
 	FCompositeFont* CompositeFont;
 	int32 SubTypefaceEntryIndex;
+	bool bRenameRequested;
 
 	FSubTypefaceListViewEntry()
 		: CompositeFont(nullptr)
 		, SubTypefaceEntryIndex(INDEX_NONE)
+		, bRenameRequested(false)
 	{
 	}
 
 	FSubTypefaceListViewEntry(FCompositeFont* const InCompositeFont, const int32 InSubTypefaceEntryIndex)
 		: CompositeFont(InCompositeFont)
 		, SubTypefaceEntryIndex(InSubTypefaceEntryIndex)
+		, bRenameRequested(false)
 	{
 	}
 
@@ -274,9 +280,13 @@ FReply SCompositeFontEditor::OnAddSubFontFamily()
 	FCompositeFont* const CompositeFontPtr = GetCompositeFont();
 	if(CompositeFontPtr)
 	{
-		CompositeFontPtr->SubTypefaces.Add(FCompositeSubFont());
+		const int32 NewSubFontIndex = CompositeFontPtr->SubTypefaces.Add(FCompositeSubFont());
 		UpdateSubTypefaceList();
 
+		// Ask for the newly added entry to be renamed to draw attention to it
+		check(SubTypefaceEntries.IsValidIndex(NewSubFontIndex));
+		SubTypefaceEntries[NewSubFontIndex]->bRenameRequested = true;
+
 		FlushCachedFont();
 	}
 
@@ -324,7 +334,7 @@ void STypefaceEditor::Construct(const FArguments& InArgs)
 
 				+SHorizontalBox::Slot()
 				[
-					SNew(SInlineEditableTextBlock)
+					SAssignNew(NameEditableTextBox, SInlineEditableTextBlock)
 					.Text(InArgs._TypefaceDisplayName)
 					.ToolTipText((InArgs._OnDisplayNameCommitted.IsBound()) ? LOCTEXT("FontFamilyNameTooltip", "The name of this font family (click to edit)") : FText::GetEmpty())
 					.Font(FEditorStyle::GetFontStyle("DetailsView.CategoryFontStyle"))
@@ -368,6 +378,11 @@ void STypefaceEditor::Refresh()
 	UpdateFontList();
 }
 
+void STypefaceEditor::RequestRename()
+{
+	NameEditableTextBox->EnterEditingMode();
+}
+
 void STypefaceEditor::UpdateFontList()
 {
 	FTypeface* const TypefacePtr = Typeface.Get(nullptr);
@@ -477,9 +492,13 @@ FReply STypefaceEditor::OnAddFont()
 			NewFontName.SetNumber(NewFontName.GetNumber() + 1);
 		}
 
-		TypefacePtr->Fonts.Add(FTypefaceEntry(NewFontName));
+		const int32 NewEntryIndex = TypefacePtr->Fonts.Add(FTypefaceEntry(NewFontName));
 		UpdateFontList();
 
+		// Ask for the newly added entry to be renamed to draw attention to it
+		check(TypefaceEntries.IsValidIndex(NewEntryIndex));
+		TypefaceEntries[NewEntryIndex]->bRenameRequested = true;
+
 		CompositeFontEditorPtr->FlushCachedFont();
 	}
 
@@ -563,7 +582,7 @@ void STypefaceEntryEditor::Construct(const FArguments& InArgs)
 			.AutoHeight()
 			.Padding(FMargin(0.0f, 0.0f, 0.0f, 4.0f))
 			[
-				SNew(SInlineEditableTextBlock)
+				SAssignNew(NameEditableTextBox, SInlineEditableTextBlock)
 				.Text(this, &STypefaceEntryEditor::GetTypefaceEntryName)
 				.ToolTipText(LOCTEXT("FontNameTooltip", "The name of this font within the font family (click to edit)"))
 				.OnTextCommitted(this, &STypefaceEntryEditor::OnTypefaceEntryNameCommitted)
@@ -597,9 +616,9 @@ void STypefaceEntryEditor::Construct(const FArguments& InArgs)
 					.ForegroundColor(FSlateColor::UseForeground())
 					.IsFocusable(false)
 					[
-						SNew(SImage)
-						.Image(FEditorStyle::GetBrush("PropertyWindow.Button_Ellipsis"))
-						.ColorAndOpacity(FSlateColor::UseForeground())
+						SNew(STextBlock)
+						.Font(FEditorStyle::Get().GetFontStyle("FontAwesome.10"))
+						.Text(FEditorFontGlyphs::Folder_Open)
 					]
 				]
 			]
@@ -645,6 +664,15 @@ void STypefaceEntryEditor::Construct(const FArguments& InArgs)
 	];
 }
 
+void STypefaceEntryEditor::Tick(const FGeometry& AllottedGeometry, const double InCurrentTime, const float InDeltaTime)
+{
+	if (TypefaceEntry.IsValid() && TypefaceEntry->bRenameRequested)
+	{
+		TypefaceEntry->bRenameRequested = false;
+		NameEditableTextBox->EnterEditingMode();
+	}
+}
+
 FText STypefaceEntryEditor::GetTypefaceEntryName() const
 {
 	FTypefaceEntry* const TypefaceEntryPtr = TypefaceEntry->GetTypefaceEntry();
@@ -928,6 +956,15 @@ void SSubTypefaceEditor::Construct(const FArguments& InArgs)
 	UpdateCharacterRangesList();
 }
 
+void SSubTypefaceEditor::Tick(const FGeometry& AllottedGeometry, const double InCurrentTime, const float InDeltaTime)
+{
+	if (SubTypeface.IsValid() && SubTypeface->bRenameRequested)
+	{
+		SubTypeface->bRenameRequested = false;
+		TypefaceEditor->RequestRename();
+	}
+}
+
 FTypeface* SSubTypefaceEditor::GetTypeface() const
 {
 	FCompositeSubFont* const SubTypefaceEntryPtr = SubTypeface->GetSubTypefaceEntry();
diff --git a/Engine/Source/Editor/FontEditor/Private/SCompositeFontEditor.h b/Engine/Source/Editor/FontEditor/Private/SCompositeFontEditor.h
--- a/Engine/Source/Editor/FontEditor/Private/SCompositeFontEditor.h
+++ b/Engine/Source/Editor/FontEditor/Private/SCompositeFontEditor.h
@@ -114,6 +114,9 @@ public:
 	/** Refresh the editor in response to an external change */
 	void Refresh();
 
+	/** Request that we begin editing the display name */
+	void RequestRename();
+
 private:
 	/** Update the list of fonts in this typeface */
 	void UpdateFontList();
@@ -139,6 +142,9 @@ private:
 	/** Internal list of font pointers for the list view (generated from TypefacePtr->Fonts) */
 	TArray<FTypefaceListViewEntryPtr> TypefaceEntries;
 
+	/** Inline editable text for the typeface display name */
+	TSharedPtr<SInlineEditableTextBlock> NameEditableTextBox;
+
 	/** Tile view widget showing the font entries (uses TypefaceEntries as its source) */
 	TSharedPtr<STileView<FTypefaceListViewEntryPtr>> TypefaceEntriesTileView;
 };
@@ -178,6 +184,7 @@ public:
 
 	/** SCompoundWidget interface */
 	void Construct(const FArguments& InArgs);
+	virtual void Tick(const FGeometry& AllottedGeometry, const double InCurrentTime, const float InDeltaTime) override;
 
 private:
 	/** Entries used by the hinting type combo box */
@@ -239,8 +246,8 @@ private:
 	/** Holds a delegate that is executed when this typeface entry name is changed - used to verify that the new name is valid */
 	FOnVerifyFontName OnVerifyFontName;
 
-	/** Text block showing the preview text for this entry */
-	TSharedPtr<STextBlock> PreviewTextBlock;
+	/** Inline editable text for the font name */
+	TSharedPtr<SInlineEditableTextBlock> NameEditableTextBox;
 
 	/** Source data for the hinting combo widget */
 	TArray<TSharedPtr<FFontHintingComboEntry>> HintingComboData;
@@ -281,6 +288,7 @@ public:
 
 	/** SCompoundWidget interface */
 	void Construct(const FArguments& InArgs);
+	virtual void Tick(const FGeometry& AllottedGeometry, const double InCurrentTime, const float InDeltaTime) override;
 
 private:
 	/** Get the typeface used by this sub-typeface */
