Fix DPI being applied twice in certain viewports

diff --git a/Engine/Source/Editor/FontEditor/Private/SFontEditorViewport.cpp b/Engine/Source/Editor/FontEditor/Private/SFontEditorViewport.cpp
--- a/Engine/Source/Editor/FontEditor/Private/SFontEditorViewport.cpp
+++ b/Engine/Source/Editor/FontEditor/Private/SFontEditorViewport.cpp
@@ -70,6 +70,9 @@ private:
 	/** Returns the positions of the scrollbars relative to the font textures */
 	FVector2D GetViewportScrollBarPositions() const;
 
+	/** Get the raw (unscaled) FontEditorViewport size */
+	FIntPoint GetDPIUnscaledViewportSizeXY() const;
+
 	void HandleWindowDPIScaleChanged(TSharedRef<SWindow> Window);
 private:
 	/** Pointer back to the Font viewport control that owns us */
@@ -120,7 +123,7 @@ void FFontEditorViewportClient::Draw(FViewport* Viewport, FCanvas* Canvas)
 	if (!FontEditorViewportPtr.Pin()->IsPreviewViewport())
 	{
 		FVector2D Ratio = FVector2D(GetViewportHorizontalScrollBarRatio(), GetViewportVerticalScrollBarRatio());
-		FVector2D ViewportSize = FVector2D(FontEditorViewportPtr.Pin()->GetViewport()->GetSizeXY().X, FontEditorViewportPtr.Pin()->GetViewport()->GetSizeXY().Y);
+		FVector2D ViewportSize = FVector2D(GetDPIUnscaledViewportSizeXY().X, GetDPIUnscaledViewportSizeXY().Y);
 		FVector2D ScrollBarPos = GetViewportScrollBarPositions();
 		int32 YOffset = (Ratio.Y > 1.0f)? ((ViewportSize.Y - (ViewportSize.Y / Ratio.Y)) * 0.5f): 0;
 		int32 YPos = YOffset - ScrollBarPos.Y;
@@ -529,7 +532,7 @@ float FFontEditorViewportClient::GetViewportVerticalScrollBarRatio() const
 	{
 		UFont* Font = FontEditorViewportPtr.Pin()->GetFontEditor().Pin()->GetFont();
 		
-		WidgetHeight = FontEditorViewportPtr.Pin()->GetViewport()->GetSizeXY().Y;
+		WidgetHeight = GetDPIUnscaledViewportSizeXY().Y;
 		
 		for (int32 Idx = 0; Idx < Font->Textures.Num(); ++Idx)
 		{
@@ -557,7 +560,7 @@ float FFontEditorViewportClient::GetViewportHorizontalScrollBarRatio() const
 		UFont* Font = FontEditorViewportPtr.Pin()->GetFontEditor().Pin()->GetFont();
 		uint32 Height = 1;
 
-		WidgetWidth = FontEditorViewportPtr.Pin()->GetViewport()->GetSizeXY().X;
+		WidgetWidth = GetDPIUnscaledViewportSizeXY().X;
 
 		if(Font->Textures.Num())
 		{
@@ -573,6 +576,11 @@ float FFontEditorViewportClient::GetViewportHorizontalScrollBarRatio() const
 	return WidgetWidth / Width;
 }
 
+FIntPoint FFontEditorViewportClient::GetDPIUnscaledViewportSizeXY() const
+{
+	return FontEditorViewportPtr.Pin()->GetViewport()->GetSizeXY() / GetDPIScale();
+}
+
 float FFontEditorViewportClient::UpdateViewportClientWindowDPIScale() const
 {
 	float DPIScale = 1.f;
diff --git a/Engine/Source/Editor/Persona/Private/AnimationEditorViewportClient.cpp b/Engine/Source/Editor/Persona/Private/AnimationEditorViewportClient.cpp
--- a/Engine/Source/Editor/Persona/Private/AnimationEditorViewportClient.cpp
+++ b/Engine/Source/Editor/Persona/Private/AnimationEditorViewportClient.cpp
@@ -597,8 +597,8 @@ void FAnimationViewportClient::ShowBoneNames( FCanvas* Canvas, FSceneView* View
 	const int32 LODIndex = FMath::Clamp(PreviewMeshComponent->PredictedLODLevel, 0, SkelMeshRenderData->LODRenderData.Num()-1);
 	FSkeletalMeshLODRenderData& LODData = SkelMeshRenderData->LODRenderData[ LODIndex ];
 
-	const int32 HalfX = Viewport->GetSizeXY().X/2;
-	const int32 HalfY = Viewport->GetSizeXY().Y/2;
+	const int32 HalfX = Viewport->GetSizeXY().X/2 / GetDPIScale();
+	const int32 HalfY = Viewport->GetSizeXY().Y/2 / GetDPIScale();
 
 	for (int32 i=0; i< LODData.RequiredBones.Num(); i++)
 	{
