Removed redundant use of FCharacterList

#rnx

diff --git a/Engine/Source/Editor/FontEditor/Private/SFontEditorViewport.cpp b/Engine/Source/Editor/FontEditor/Private/SFontEditorViewport.cpp
--- a/Engine/Source/Editor/FontEditor/Private/SFontEditorViewport.cpp
+++ b/Engine/Source/Editor/FontEditor/Private/SFontEditorViewport.cpp
@@ -224,8 +224,6 @@ void FFontEditorViewportClient::Draw(FViewport* Viewport, FCanvas* Canvas)
 			{
 				const FSlateFontInfo FontInfo(Font, Font->LegacyFontSize, TypefaceEntry.Name);
 
-				FCharacterList& CharacterList = FontCache->GetCharacterList(FontInfo, FontScale);
-
 				FShapedGlyphSequenceRef EntryNameShapedText = FontCache->ShapeBidirectionalText(TypefaceEntry.Name.ToString(), FontInfo, FontScale, TextBiDi::ETextDirection::LeftToRight, GetDefaultTextShapingMethod());
 				
 				FCanvasShapedTextItem ShapedTextItem(CurPos, EntryNameShapedText, FLinearColor(ForegroundColor));
@@ -234,7 +232,7 @@ void FFontEditorViewportClient::Draw(FViewport* Viewport, FCanvas* Canvas)
 				const FVector2D MeasuredText = ShapedTextItem.DrawnSize;
 				WidestName = FMath::Max(WidestName, EntryNameShapedText->GetMeasuredWidth());
 
-				CurPos.Y += CharacterList.GetMaxHeight() + 8.0f;
+				CurPos.Y += EntryNameShapedText->GetMaxTextHeight() + 8.0f;
 			}
 
 			CurPos = FVector2D(WidestName + 12.0f, StartPos.Y);
