Fixed a crash and editor issue when importing a font face over an existing asset

diff --git a/Engine/Source/Editor/FontEditor/Private/FontFaceDetailsCustomization.cpp b/Engine/Source/Editor/FontEditor/Private/FontFaceDetailsCustomization.cpp
--- a/Engine/Source/Editor/FontEditor/Private/FontFaceDetailsCustomization.cpp
+++ b/Engine/Source/Editor/FontEditor/Private/FontFaceDetailsCustomization.cpp
@@ -80,7 +80,7 @@ FText FFontFaceDetailsCustomization::GetFontDisplayName() const
 {
 	if (ObjectsBeingEdited.Num() == 1)
 	{
-		const UFontFace* FontFace = CastChecked<UFontFace>(ObjectsBeingEdited[0].Get());
+		const UFontFace* FontFace = Cast<UFontFace>(ObjectsBeingEdited[0].Get());
 		if (FontFace && !FontFace->SourceFilename.IsEmpty())
 		{
 			return FText::FromString(FPaths::GetCleanFilename(FontFace->SourceFilename));
@@ -94,7 +94,7 @@ FText FFontFaceDetailsCustomization::GetFontDisplayToolTip() const
 {
 	if (ObjectsBeingEdited.Num() == 1)
 	{
-		const UFontFace* FontFace = CastChecked<UFontFace>(ObjectsBeingEdited[0].Get());
+		const UFontFace* FontFace = Cast<UFontFace>(ObjectsBeingEdited[0].Get());
 		if (FontFace && !FontFace->SourceFilename.IsEmpty())
 		{
 			return FText::FromString(FontFace->SourceFilename);
@@ -139,7 +139,7 @@ void FFontFaceDetailsCustomization::OnFontPathPicked(const FString& InNewFontFil
 
 	for (const TWeakObjectPtr<UObject>& ObjectBeingEdited : ObjectsBeingEdited)
 	{
-		UFontFace* FontFace = CastChecked<UFontFace>(ObjectBeingEdited.Get());
+		UFontFace* FontFace = Cast<UFontFace>(ObjectBeingEdited.Get());
 		if (FontFace)
 		{
 			FontFace->Modify();
diff --git a/Engine/Source/Editor/FontEditor/Private/SCompositeFontEditor.cpp b/Engine/Source/Editor/FontEditor/Private/SCompositeFontEditor.cpp
--- a/Engine/Source/Editor/FontEditor/Private/SCompositeFontEditor.cpp
+++ b/Engine/Source/Editor/FontEditor/Private/SCompositeFontEditor.cpp
@@ -27,7 +27,9 @@
 #include "Misc/FileHelper.h"
 #include "IContentBrowserSingleton.h"
 #include "FileHelpers.h"
+#include "Subsystems/ImportSubsystem.h"
 #include "Subsystems/AssetEditorSubsystem.h"
+#include "Factories/FontFileImportFactory.h"
 
 #define LOCTEXT_NAMESPACE "FontEditor"
 
@@ -1100,8 +1102,10 @@ UFontFace* STypefaceEntryEditor::SaveFontFaceAsAsset(const UFontFace* InFontFace
 	{
 		const FString NewFaceAssetName = FPackageName::GetLongPackageAssetName(NewPackageName);
 		UPackage* NewFaceAssetPackage = CreatePackage(nullptr, *NewPackageName);
-		UFontFace* NewFaceAsset = Cast<UFontFace>(StaticDuplicateObject(InFontFace, NewFaceAssetPackage, *NewFaceAssetName));
 
+		GEditor->GetEditorSubsystem<UImportSubsystem>()->BroadcastAssetPreImport(GetMutableDefault<UFontFileImportFactory>(), UFontFace::StaticClass(), NewFaceAssetPackage, *NewFaceAssetName, *FPaths::GetExtension(InFontFace->GetFontFilename()));
+
+		UFontFace* NewFaceAsset = Cast<UFontFace>(StaticDuplicateObject(InFontFace, NewFaceAssetPackage, *NewFaceAssetName));
 		if (NewFaceAsset)
 		{
 			// Make sure the new object is flagged correctly
@@ -1111,6 +1115,8 @@ UFontFace* STypefaceEntryEditor::SaveFontFaceAsAsset(const UFontFace* InFontFace
 			FAssetRegistryModule::AssetCreated(NewFaceAsset);
 		}
 
+		GEditor->GetEditorSubsystem<UImportSubsystem>()->BroadcastAssetPostImport(GetMutableDefault<UFontFileImportFactory>(), NewFaceAsset);
+
 		return NewFaceAsset;
 	}
 
