PR #4860: UE-36822: Give me user feedback when combobox is useable (Contributed by [REDACTED_PERSON])

diff --git a/Engine/Source/Editor/FontEditor/Private/SCompositeFontEditor.cpp b/Engine/Source/Editor/FontEditor/Private/SCompositeFontEditor.cpp
--- a/Engine/Source/Editor/FontEditor/Private/SCompositeFontEditor.cpp
+++ b/Engine/Source/Editor/FontEditor/Private/SCompositeFontEditor.cpp
@@ -1761,7 +1761,7 @@ void SFontScalingFactorEditor::Construct(const FArguments& InArgs)
 		.AutoWidth()
 		[
 			SNew(SNumericEntryBox<float>)
-			.ToolTipText(LOCTEXT("ScalingFactorTooltip", "The scaling factor will adjust the size of the rendered glyphs so that you can tweak their size to match that of the default font family"))
+			.ToolTipText(LOCTEXT("ScalingFactorTooltip", "The scaling factor will adjust the size of the rendered glyphs so that you can tweak their size to match that of the default font family\nNote: Only applies to scalable font formats (ie, not to bitmap fonts)"))
 			.Value(this, &SFontScalingFactorEditor::GetScalingFactorAsOptional)
 			.OnValueCommitted(this, &SFontScalingFactorEditor::OnScalingFactorCommittedAsNumeric)
 		]
@@ -1808,6 +1808,7 @@ void SFontOverrideSelector::Construct(const FArguments& InArgs)
 		.OptionsSource(&FontOverrideComboData)
 		.ContentPadding(FMargin(4.0, 2.0))
 		.Visibility(this, &SFontOverrideSelector::GetAddFontOverrideVisibility)
+		.IsEnabled(this, &SFontOverrideSelector::IsFontOverrideComboEnabled)
 		.OnComboBoxOpening(this, &SFontOverrideSelector::OnAddFontOverrideComboOpening)
 		.OnSelectionChanged(this, &SFontOverrideSelector::OnAddFontOverrideSelectionChanged)
 		.OnGenerateWidget(this, &SFontOverrideSelector::MakeAddFontOverrideWidget)
@@ -1824,6 +1825,31 @@ EVisibility SFontOverrideSelector::GetAddFontOverrideVisibility() const
 	return (ParentTypeface.Get(nullptr)) ? EVisibility::Visible : EVisibility::Collapsed;
 }
 
+bool SFontOverrideSelector::IsFontOverrideComboEnabled() const
+{
+	FTypeface* const TypefacePtr = Typeface.Get(nullptr);
+	const FTypeface* const ParentTypefacePtr = ParentTypeface.Get(nullptr);
+
+	if (TypefacePtr && ParentTypefacePtr)
+	{
+		// Check whether our parent font has any entries that haven't already got a local entry
+		for (const FTypefaceEntry& ParentTypefaceEntry : ParentTypefacePtr->Fonts)
+		{
+			const bool bIsOverridden = TypefacePtr->Fonts.ContainsByPredicate([&ParentTypefaceEntry](const FTypefaceEntry& LocalTypefaceEntry)
+			{
+				return LocalTypefaceEntry.Name == ParentTypefaceEntry.Name;
+			});
+
+			if (!bIsOverridden)
+			{
+				return true;
+			}
+		}
+	}
+
+	return false;
+}
+
 void SFontOverrideSelector::OnAddFontOverrideComboOpening()
 {
 	FontOverrideComboData.Reset();
diff --git a/Engine/Source/Editor/FontEditor/Private/SCompositeFontEditor.h b/Engine/Source/Editor/FontEditor/Private/SCompositeFontEditor.h
--- a/Engine/Source/Editor/FontEditor/Private/SCompositeFontEditor.h
+++ b/Engine/Source/Editor/FontEditor/Private/SCompositeFontEditor.h
@@ -534,6 +534,9 @@ private:
 	/** Get the visibility of the "Add Font Override" combo button */
 	EVisibility GetAddFontOverrideVisibility() const;
 
+	/** True if the font override combo should be enabled (due to still having parent fonts that can be overridden in this sub-font) */
+	bool IsFontOverrideComboEnabled() const;
+
 	/** Called before the font override combo is opened - used to update the list of available font overrides */
 	void OnAddFontOverrideComboOpening();
 
