Fixes Font Editor Spacing Issues

diff --git a/Engine/Source/Editor/FontEditor/Private/SCompositeFontEditor.cpp b/Engine/Source/Editor/FontEditor/Private/SCompositeFontEditor.cpp
--- a/Engine/Source/Editor/FontEditor/Private/SCompositeFontEditor.cpp
+++ b/Engine/Source/Editor/FontEditor/Private/SCompositeFontEditor.cpp
@@ -1526,6 +1526,7 @@ void SCharacterRangeEditor::Construct(const FArguments& InArgs)
 				SNew(SVerticalBox)
 
 				+SVerticalBox::Slot()
+				.AutoHeight()
 				.Padding(2.0f)
 				[
 					SNew(SEditableTextBox)
@@ -1535,6 +1536,7 @@ void SCharacterRangeEditor::Construct(const FArguments& InArgs)
 				]
 
 				+SVerticalBox::Slot()
+				.AutoHeight()
 				.Padding(2.0f)
 				[
 					SNew(SEditableTextBox)
@@ -1544,6 +1546,7 @@ void SCharacterRangeEditor::Construct(const FArguments& InArgs)
 				]
 
 				+SVerticalBox::Slot()
+				.AutoHeight()
 				.Padding(2.0f)
 				[
 					SNew(SNumericEntryBox<int32>)
@@ -1570,6 +1573,7 @@ void SCharacterRangeEditor::Construct(const FArguments& InArgs)
 				SNew(SVerticalBox)
 
 				+SVerticalBox::Slot()
+				.AutoHeight()
 				.Padding(2.0f)
 				[
 					SNew(SEditableTextBox)
@@ -1579,6 +1583,7 @@ void SCharacterRangeEditor::Construct(const FArguments& InArgs)
 				]
 
 				+SVerticalBox::Slot()
+				.AutoHeight()
 				.Padding(2.0f)
 				[
 					SNew(SEditableTextBox)
@@ -1588,6 +1593,7 @@ void SCharacterRangeEditor::Construct(const FArguments& InArgs)
 				]
 
 				+SVerticalBox::Slot()
+				.AutoHeight()
 				.Padding(2.0f)
 				[
 					SNew(SNumericEntryBox<int32>)
