Fixes Spacing in Font Editor

diff --git a/Engine/Source/Editor/FontEditor/Private/SCompositeFontEditor.cpp b/Engine/Source/Editor/FontEditor/Private/SCompositeFontEditor.cpp
--- a/Engine/Source/Editor/FontEditor/Private/SCompositeFontEditor.cpp
+++ b/Engine/Source/Editor/FontEditor/Private/SCompositeFontEditor.cpp
@@ -459,7 +459,7 @@ void STypefaceEditor::Construct(const FArguments& InArgs)
 				SAssignNew(TypefaceEntriesTileView, STileView<FTypefaceListViewEntryPtr>)
 				.ListItemsSource(&TypefaceEntries)
 				.SelectionMode(ESelectionMode::None)
-				.ItemWidth(160)
+				.ItemWidth(180)
 				.ItemHeight(120)
 				.ItemAlignment(EListItemAlignment::LeftAligned)
 				.OnGenerateTile(this, &STypefaceEditor::MakeTypefaceEntryWidget)
@@ -1226,8 +1226,8 @@ void SSubTypefaceEditor::Construct(const FArguments& InArgs)
 				SAssignNew(CharacterRangeEntriesTileView, STileView<FCharacterRangeTileViewEntryPtr>)
 				.ListItemsSource(&CharacterRangeEntries)
 				.SelectionMode(ESelectionMode::None)
-				.ItemWidth(160)
-				.ItemHeight(144)
+				.ItemWidth(180)
+				.ItemHeight(160)
 				.ItemAlignment(EListItemAlignment::LeftAligned)
 				.OnGenerateTile(this, &SSubTypefaceEditor::MakeCharacterRangesEntryWidget)
 			]
