#ue5 - Fix truncation warnings in modules: DeviceProfileServices, FieldSystemEngine, FontEditor, FunctionalTesting. Includes some related engine fixes as well.

diff --git a/Engine/Source/Developer/FunctionalTesting/Classes/FunctionalTest.h b/Engine/Source/Developer/FunctionalTesting/Classes/FunctionalTest.h
--- a/Engine/Source/Developer/FunctionalTesting/Classes/FunctionalTest.h
+++ b/Engine/Source/Developer/FunctionalTesting/Classes/FunctionalTest.h
@@ -35,7 +35,7 @@ struct FStatsData
 	FStatsData() :NumFrames(0), SumTimeSeconds(0.0f){}
 
 	uint32 NumFrames;
-	uint32 SumTimeSeconds;
+	float SumTimeSeconds;
 	FStatisticalFloat FrameTimeTracker;
 	FStatisticalFloat GameThreadTimeTracker;
 	FStatisticalFloat RenderThreadTimeTracker;
diff --git a/Engine/Source/Developer/FunctionalTesting/Private/AutomationBlueprintFunctionLibrary.cpp b/Engine/Source/Developer/FunctionalTesting/Private/AutomationBlueprintFunctionLibrary.cpp
--- a/Engine/Source/Developer/FunctionalTesting/Private/AutomationBlueprintFunctionLibrary.cpp
+++ b/Engine/Source/Developer/FunctionalTesting/Private/AutomationBlueprintFunctionLibrary.cpp
@@ -964,11 +964,11 @@ float HelperGetStat(FName StatName)
 		{
 			if(bCallCount)
 			{
-				return StatMessage->GetValue_CallCount(ValueType);
+				return (float)StatMessage->GetValue_CallCount(ValueType);
 			}
 			else
 			{
-				return FPlatformTime::ToMilliseconds(StatMessage->GetValue_Duration(ValueType));
+				return (float)FPlatformTime::ToMilliseconds64(StatMessage->GetValue_Duration(ValueType));
 			}
 		}
 	}
diff --git a/Engine/Source/Developer/FunctionalTesting/Private/FuncTestRenderingComponent.cpp b/Engine/Source/Developer/FunctionalTesting/Private/FuncTestRenderingComponent.cpp
--- a/Engine/Source/Developer/FunctionalTesting/Private/FuncTestRenderingComponent.cpp
+++ b/Engine/Source/Developer/FunctionalTesting/Private/FuncTestRenderingComponent.cpp
@@ -70,7 +70,7 @@ public:
 	}
 
 	virtual uint32 GetMemoryFootprint( void ) const override { return( sizeof( *this ) + GetAllocatedSize() ); }
-	uint32 GetAllocatedSize( void ) const { return FPrimitiveSceneProxy::GetAllocatedSize() + Locations.GetAllocatedSize(); }
+	SIZE_T GetAllocatedSize( void ) const { return FPrimitiveSceneProxy::GetAllocatedSize() + Locations.GetAllocatedSize(); }
 
 protected:
 	AFunctionalTest* OwningTest;
diff --git a/Engine/Source/Developer/FunctionalTesting/Private/FunctionalTest.cpp b/Engine/Source/Developer/FunctionalTesting/Private/FunctionalTest.cpp
--- a/Engine/Source/Developer/FunctionalTesting/Private/FunctionalTest.cpp
+++ b/Engine/Source/Developer/FunctionalTesting/Private/FunctionalTest.cpp
@@ -259,7 +259,7 @@ bool AFunctionalTest::RunTest(const TArray<FString>& Params)
 	//GEngine->DelayGarbageCollection();
 
 	RunFrame = GFrameNumber;
-	RunTime = World->GetTimeSeconds();
+	RunTime = (float)World->GetTimeSeconds();
 
 	TotalTime = 0.f;
 	if (TimeLimit >= 0)
@@ -289,7 +289,7 @@ void AFunctionalTest::StartTest()
 	SCOPE_CYCLE_COUNTER(STAT_FunctionalTest_StartTest);
 	TotalTime = 0.f;
 	StartFrame = GFrameNumber;
-	StartTime = GetWorld()->GetTimeSeconds();
+	StartTime = (float)GetWorld()->GetTimeSeconds();
 
 	ReceiveStartTest();
 	OnTestStart.Broadcast();
@@ -1282,14 +1282,14 @@ FString FPerfStatsRecord::GetOverBudgetString() const
 {
 	double Min, Max, Avg;
 	GetRenderThreadTimes(Min, Max, Avg);
-	float RTMax = Max;
-	float RTBudgetFrac = Max / RenderThreadBudget;
+	float RTMax = (float)Max;
+	float RTBudgetFrac = (float)(Max / RenderThreadBudget);
 	GetGameThreadTimes(Min, Max, Avg);
-	float GTMax = Max;
-	float GTBudgetFrac = Max / GameThreadBudget;
+	float GTMax = (float)Max;
+	float GTBudgetFrac = (float)(Max / GameThreadBudget);
 	GetGPUTimes(Min, Max, Avg);
-	float GPUMax = Max;
-	float GPUBudgetFrac = Max / GPUBudget;
+	float GPUMax = (float)Max;
+	float GPUBudgetFrac = (float)(Max / GPUBudget);
 
 	return FString::Printf(TEXT("%s,%.2f,%.2f,%.2f,%.2f,%.2f,%.2f,%.2f,%.2f,%.2f"),
 		*Name,
diff --git a/Engine/Source/Editor/DeviceProfileServices/Private/CheckAndroidDeviceProfileCommandlet.cpp b/Engine/Source/Editor/DeviceProfileServices/Private/CheckAndroidDeviceProfileCommandlet.cpp
--- a/Engine/Source/Editor/DeviceProfileServices/Private/CheckAndroidDeviceProfileCommandlet.cpp
+++ b/Engine/Source/Editor/DeviceProfileServices/Private/CheckAndroidDeviceProfileCommandlet.cpp
@@ -178,7 +178,7 @@ int32 UCheckAndroidDeviceProfileCommandlet::Main(const FString& RawCommandLine)
 			TArray<FSelectedFragmentProperties> SelectedFragments;
 
 			UE_LOG(LogCheckAndroidDeviceProfile, Display, TEXT("%s Selected Device Profile: %s"), *DeviceName, *ProfileName);
-			int32 TotalPhysGB = FCString::Atoi64(*DeviceParameters.FindChecked(TEXT("SRC_TotalPhysicalGB")));
+			int32 TotalPhysGB = FCString::Atoi(*DeviceParameters.FindChecked(TEXT("SRC_TotalPhysicalGB")));
 			EPlatformMemorySizeBucket MemorySizeBucket = DetermineDeviceMemorySizeBucket(TotalPhysGB);
 			FString PlatformOverride(TEXT("Android"));
 			FString ProfileDescription;
diff --git a/Engine/Source/Editor/FontEditor/Private/SFontEditorViewport.cpp b/Engine/Source/Editor/FontEditor/Private/SFontEditorViewport.cpp
--- a/Engine/Source/Editor/FontEditor/Private/SFontEditorViewport.cpp
+++ b/Engine/Source/Editor/FontEditor/Private/SFontEditorViewport.cpp
@@ -155,14 +155,15 @@ void FFontEditorViewportClient::Draw(FViewport* Viewport, FCanvas* Canvas)
 
 	if (!FontEditorViewportPtr.Pin()->IsPreviewViewport())
 	{
-		FVector2D Ratio = FVector2D(GetViewportHorizontalScrollBarRatio(), GetViewportVerticalScrollBarRatio());
-		FVector2D ViewportSize = FVector2D(GetDPIUnscaledViewportSizeXY().X, GetDPIUnscaledViewportSizeXY().Y);
-		FVector2D ScrollBarPos = GetViewportScrollBarPositions();
-		int32 YOffset = (Ratio.Y > 1.0f)? ((ViewportSize.Y - (ViewportSize.Y / Ratio.Y)) * 0.5f): 0;
-		int32 YPos = YOffset - ScrollBarPos.Y;
+		FVector2f Ratio = FVector2f(GetViewportHorizontalScrollBarRatio(), GetViewportVerticalScrollBarRatio());
+		const FIntPoint DPIUnscaledViewportSize = GetDPIUnscaledViewportSizeXY();
+		FVector2f ViewportSize = FVector2f((float)DPIUnscaledViewportSize.X, (float)DPIUnscaledViewportSize.Y);
+		FVector2f ScrollBarPos = FVector2f(GetViewportScrollBarPositions());
+		int32 YOffset = FMath::TruncToInt32((Ratio.Y > 1.0f)? ((ViewportSize.Y - (ViewportSize.Y / Ratio.Y)) * 0.5f): 0);
+		int32 YPos = FMath::TruncToInt32(YOffset - ScrollBarPos.Y);
 		int32 LastDrawnYPos = YPos;
-		int32 XOffset = (Ratio.X > 1.0f)? ((ViewportSize.X - (ViewportSize.X / Ratio.X)) * 0.5f): 0;
-		int32 XPos = XOffset - ScrollBarPos.X;
+		int32 XOffset = FMath::TruncToInt32((Ratio.X > 1.0f)? ((ViewportSize.X - (ViewportSize.X / Ratio.X)) * 0.5f): 0);
+		int32 XPos = FMath::TruncToInt32(XOffset - ScrollBarPos.X);
 	
 		UpdateScrollBars();
 
@@ -172,7 +173,7 @@ void FFontEditorViewportClient::Draw(FViewport* Viewport, FCanvas* Canvas)
 		const static float BoxSize = 40.0f;
 		bool XColor = false;
 		bool YColor = false;
-		FVector2D BackgroundPos;
+		FVector2f BackgroundPos;
 		for (BackgroundPos.Y = 0.0f; BackgroundPos.Y <= ViewportSize.Y; BackgroundPos.Y += BoxSize)
 		{
 			YColor = !YColor;
@@ -211,7 +212,7 @@ void FFontEditorViewportClient::Draw(FViewport* Viewport, FCanvas* Canvas)
 							Canvas->SetHitProxy(new HObject(Texture));
 						
 							// Draw a simple tile
-							Canvas->DrawTile( XPos, YPos, Width, Height, 0.0f, 0.0f, 1.0f, 1.0f, FLinearColor::White, nullptr );
+							Canvas->DrawTile( (float)XPos, (float)YPos, (float)Width, (float)Height, 0.0f, 0.0f, 1.0f, 1.0f, FLinearColor::White, nullptr );
 						
 							Canvas->SetHitProxy(NULL);
 						}
@@ -221,15 +222,15 @@ void FFontEditorViewportClient::Draw(FViewport* Viewport, FCanvas* Canvas)
 							// Draw a selected background
 							if (Texture->IsSelected())
 							{
-								Canvas->DrawTile( XPos, YPos, Width, Height, 0.0f, 0.0f, 1.0f, 1.0f, FLinearColor(0.084f, .127f, 0.098f), nullptr );
+								Canvas->DrawTile((float)XPos, (float)YPos, (float)Width, (float)Height, 0.0f, 0.0f, 1.0f, 1.0f, FLinearColor(0.084f, .127f, 0.098f), nullptr );
 							}
 							else
 							{
-								Canvas->DrawTile( XPos, YPos, Width, Height, 0.0f, 0.0f, 1.0f, 1.0f, FLinearColor::Black);
+								Canvas->DrawTile((float)XPos, (float)YPos, (float)Width, (float)Height, 0.0f, 0.0f, 1.0f, 1.0f, FLinearColor::Black);
 							}
 
 							// Draw the font texture (with alpha blending enabled)
-							Canvas->DrawTile( XPos, YPos, Width, Height, 0.0f, 0.0f, 1.0f, 1.0f, FLinearColor::White, Texture->GetResource(), true);
+							Canvas->DrawTile((float)XPos, (float)YPos, (float)Width, (float)Height, 0.0f, 0.0f, 1.0f, 1.0f, FLinearColor::White, Texture->GetResource(), true);
 						}
 					}
 					// Update our total height and current draw position
@@ -299,8 +300,8 @@ void FFontEditorViewportClient::Draw(FViewport* Viewport, FCanvas* Canvas)
 								const float BitmapRenderScale = GlyphToRender.GetBitmapRenderScale();
 								const float InvBitmapRenderScale = 1.0f / BitmapRenderScale;
 
-								const float X = CurPos.X + LineX + GlyphAtlasData.HorizontalOffset + GlyphToRender.XOffset;
-								const float Y = CurPos.Y - GlyphAtlasData.VerticalOffset + GlyphToRender.YOffset + ((ShapedPreviewText->GetTextBaseline() + ShapedPreviewText->GetMaxTextHeight()) * InvBitmapRenderScale);
+								const double X = CurPos.X + LineX + GlyphAtlasData.HorizontalOffset + GlyphToRender.XOffset;
+								const double Y = CurPos.Y - GlyphAtlasData.VerticalOffset + GlyphToRender.YOffset + ((ShapedPreviewText->GetTextBaseline() + ShapedPreviewText->GetMaxTextHeight()) * InvBitmapRenderScale);
 
 								FCanvasBoxItem BoundingBoxItem(FVector2D(X, Y), FVector2D(GlyphAtlasData.USize * BitmapRenderScale, GlyphAtlasData.VSize * BitmapRenderScale));
 								BoundingBoxItem.SetColor(FColor::Orange);
@@ -332,8 +333,8 @@ void FFontEditorViewportClient::Draw(FViewport* Viewport, FCanvas* Canvas)
 									const float BitmapRenderScale = GlyphToRender.GetBitmapRenderScale();
 									const float InvBitmapRenderScale = 1.0f / BitmapRenderScale;
 
-									const float X = CurPos.X + LineX + GlyphClusterAdvance + GlyphAtlasData.HorizontalOffset + GlyphToRender.XOffset;
-									const float Y = CurPos.Y - GlyphAtlasData.VerticalOffset + GlyphToRender.YOffset + ((ShapedPreviewText->GetTextBaseline() + ShapedPreviewText->GetMaxTextHeight()) * InvBitmapRenderScale);
+									const double X = CurPos.X + LineX + GlyphClusterAdvance + GlyphAtlasData.HorizontalOffset + GlyphToRender.XOffset;
+									const double Y = CurPos.Y - GlyphAtlasData.VerticalOffset + GlyphToRender.YOffset + ((ShapedPreviewText->GetTextBaseline() + ShapedPreviewText->GetMaxTextHeight()) * InvBitmapRenderScale);
 
 									FVector2D ExtraWidth = FVector2D(ForceInitToZero);
 									if (GlyphClusterBounds.bIsValid)
@@ -366,7 +367,7 @@ void FFontEditorViewportClient::Draw(FViewport* Viewport, FCanvas* Canvas)
 
 					// Draw the baseline
 					{
-						const float Y = CurPos.Y + ShapedPreviewText->GetTextBaseline() + ShapedPreviewText->GetMaxTextHeight();
+						const double Y = CurPos.Y + ShapedPreviewText->GetTextBaseline() + ShapedPreviewText->GetMaxTextHeight();
 
 						FCanvasLineItem BaseLineItem(FVector2D(CurPos.X, Y), FVector2D(CurPos.X + ShapedPreviewText->GetMeasuredWidth(), Y));
 						BaseLineItem.SetColor(FLinearColor::Red);
@@ -571,7 +572,7 @@ float FFontEditorViewportClient::GetViewportVerticalScrollBarRatio() const
 	{
 		UFont* Font = FontEditorViewportPtr.Pin()->GetFontEditor().Pin()->GetFont();
 		
-		WidgetHeight = GetDPIUnscaledViewportSizeXY().Y;
+		WidgetHeight = (float)GetDPIUnscaledViewportSizeXY().Y;
 		
 		for (int32 Idx = 0; Idx < Font->Textures.Num(); ++Idx)
 		{
@@ -599,7 +600,7 @@ float FFontEditorViewportClient::GetViewportHorizontalScrollBarRatio() const
 		UFont* Font = FontEditorViewportPtr.Pin()->GetFontEditor().Pin()->GetFont();
 		uint32 Height = 1;
 
-		WidgetWidth = GetDPIUnscaledViewportSizeXY().X;
+		WidgetWidth = (float)GetDPIUnscaledViewportSizeXY().X;
 
 		if(Font->Textures.Num())
 		{
@@ -617,7 +618,9 @@ float FFontEditorViewportClient::GetViewportHorizontalScrollBarRatio() const
 
 FIntPoint FFontEditorViewportClient::GetDPIUnscaledViewportSizeXY() const
 {
-	return FontEditorViewportPtr.Pin()->GetViewport()->GetSizeXY() / GetDPIScale();
+	const FIntPoint SizeXY = FontEditorViewportPtr.Pin()->GetViewport()->GetSizeXY();
+	const float DPIScale = GetDPIScale();
+	return FIntPoint(FMath::TruncToInt32(SizeXY.X / DPIScale), FMath::TruncToInt32(SizeXY.Y / DPIScale));
 }
 
 float FFontEditorViewportClient::UpdateViewportClientWindowDPIScale() const
diff --git a/Engine/Source/Runtime/Core/Public/Math/Vector2D.h b/Engine/Source/Runtime/Core/Public/Math/Vector2D.h
--- a/Engine/Source/Runtime/Core/Public/Math/Vector2D.h
+++ b/Engine/Source/Runtime/Core/Public/Math/Vector2D.h
@@ -1179,7 +1179,15 @@ PRAGMA_ENABLE_DEPRECATION_WARNINGS
 template<typename T>
 FORCEINLINE FIntPoint TVector2<T>::IntPoint() const
 {
-	return FIntPoint(FMath::RoundToInt(X), FMath::RoundToInt(Y));
+	if constexpr (std::is_same_v<T, float>)
+	{
+		return FIntPoint( FMath::RoundToInt32(X), FMath::RoundToInt32(Y) );
+	}
+	else
+	{
+		// FIntPoint constructor from FInt64Point checks that the int64 fits in int32.
+		return FIntPoint( FInt64Point(FMath::RoundToInt64(X), FMath::RoundToInt64(Y)) );
+	}
 }
 
 template<typename T>
diff --git a/Engine/Source/Runtime/Engine/Public/GestureRecognizer.h b/Engine/Source/Runtime/Engine/Public/GestureRecognizer.h
--- a/Engine/Source/Runtime/Engine/Public/GestureRecognizer.h
+++ b/Engine/Source/Runtime/Engine/Public/GestureRecognizer.h
@@ -14,7 +14,7 @@ class FGestureRecognizer
 public:
 	/** Constructor **/
 	FGestureRecognizer() : bIsReadyForPinch(false), AnchorDistanceSq(0.0f), StartAngle(0.0f),
-		bIsReadyForFlick(false), FlickTime(0.0f), PreviousTouchCount(0.0f)
+		bIsReadyForFlick(false), FlickTime(0.0f), PreviousTouchCount(0)
 	{};
 
 	/** Attempt to detect touch gestures */
diff --git a/Engine/Source/Runtime/Engine/Public/Tests/AutomationCommon.h b/Engine/Source/Runtime/Engine/Public/Tests/AutomationCommon.h
--- a/Engine/Source/Runtime/Engine/Public/Tests/AutomationCommon.h
+++ b/Engine/Source/Runtime/Engine/Public/Tests/AutomationCommon.h
@@ -231,7 +231,7 @@ private:
 	const double kTickRate = 60.0;
 
 	// How many samples we hold
-	const int kSampleCount = kTickRate * 5;
+	const int kSampleCount = (int32)kTickRate * 5;
 };
 
 /**
diff --git a/Engine/Source/Runtime/Experimental/Chaos/Public/Field/FieldSystem.h b/Engine/Source/Runtime/Experimental/Chaos/Public/Field/FieldSystem.h
--- a/Engine/Source/Runtime/Experimental/Chaos/Public/Field/FieldSystem.h
+++ b/Engine/Source/Runtime/Experimental/Chaos/Public/Field/FieldSystem.h
@@ -561,10 +561,10 @@ public:
 		MetaData[Key].Reset(Value);
 	}
 
-	void InitFieldNodes(const float& TimeSeconds, const FName& Name)
+	void InitFieldNodes(const double TimeSeconds, const FName& Name)
 	{
 		CommandName = Name;
-		TimeCreation = TimeSeconds;
+		TimeCreation = (float)TimeSeconds;
 	}
 
 	void Serialize(FArchive& Ar);
diff --git a/Engine/Source/Runtime/Experimental/FieldSystem/Source/FieldSystemEngine/Private/Field/FieldSystemSceneProxy.h b/Engine/Source/Runtime/Experimental/FieldSystem/Source/FieldSystemEngine/Private/Field/FieldSystemSceneProxy.h
--- a/Engine/Source/Runtime/Experimental/FieldSystem/Source/FieldSystemEngine/Private/Field/FieldSystemSceneProxy.h
+++ b/Engine/Source/Runtime/Experimental/FieldSystem/Source/FieldSystemEngine/Private/Field/FieldSystemSceneProxy.h
@@ -36,7 +36,7 @@ public:
 	virtual uint32 GetMemoryFootprint(void) const override { return(sizeof(*this) + GetAllocatedSize()); }
 
 	/** Size of the base class */
-	uint32 GetAllocatedSize(void) const { return(FPrimitiveSceneProxy::GetAllocatedSize()); }
+	SIZE_T GetAllocatedSize(void) const { return(FPrimitiveSceneProxy::GetAllocatedSize()); }
 
 
 };
diff --git a/Engine/Source/Runtime/Experimental/FieldSystem/Source/FieldSystemEngine/Public/Field/FieldSystemObjects.h b/Engine/Source/Runtime/Experimental/FieldSystem/Source/FieldSystemEngine/Public/Field/FieldSystemObjects.h
--- a/Engine/Source/Runtime/Experimental/FieldSystem/Source/FieldSystemEngine/Public/Field/FieldSystemObjects.h
+++ b/Engine/Source/Runtime/Experimental/FieldSystem/Source/FieldSystemEngine/Public/Field/FieldSystemObjects.h
@@ -217,8 +217,8 @@ public:
 		: Super()
 		, Radius(0)
 		, Position(FVector(0, 0, 0))
-		, InteriorValue(1.0)
-		, ExteriorValue(0.0)
+		, InteriorValue(1)
+		, ExteriorValue(0)
 		, SetMaskCondition(ESetMaskConditionType::Field_Set_Always)
 	{}
 	virtual ~URadialIntMask() {}
diff --git a/Engine/Source/Runtime/Slate/Public/Widgets/Input/SSpinBox.h b/Engine/Source/Runtime/Slate/Public/Widgets/Input/SSpinBox.h
--- a/Engine/Source/Runtime/Slate/Public/Widgets/Input/SSpinBox.h
+++ b/Engine/Source/Runtime/Slate/Public/Widgets/Input/SSpinBox.h
@@ -219,11 +219,11 @@ public:
 
 		if (SupportDynamicSliderMaxValue.Get() && CachedExternalValue > GetMaxSliderValue())
 		{
-			ApplySliderMaxValueChanged(CachedExternalValue - GetMaxSliderValue(), true);
+			ApplySliderMaxValueChanged(float(CachedExternalValue - GetMaxSliderValue()), true);
 		}
 		else if (SupportDynamicSliderMinValue.Get() && CachedExternalValue < GetMinSliderValue())
 		{
-			ApplySliderMinValueChanged(CachedExternalValue - GetMinSliderValue(), true);
+			ApplySliderMinValueChanged(float(CachedExternalValue - GetMinSliderValue()), true);
 		}
 
 		UpdateIsSpinRangeUnlimited();
@@ -539,7 +539,7 @@ public:
 
 			if (!bDragging)
 			{
-				DistanceDragged += FMath::Abs(MouseEvent.GetCursorDelta().X);
+				DistanceDragged += (float)FMath::Abs(MouseEvent.GetCursorDelta().X);
 				if (DistanceDragged > FSlateApplication::Get().GetDragTriggerDistance())
 				{
 					ExitTextMode();
@@ -558,7 +558,7 @@ public:
 
 				// A minimum slider width to use for calculating deltas in the slider-range space
 				const float MinSliderWidth = 100.f;
-				float SliderWidthInSlateUnits = FMath::Max(MyGeometry.GetDrawSize().X, MinSliderWidth);
+				float SliderWidthInSlateUnits = FMath::Max((float)MyGeometry.GetDrawSize().X, MinSliderWidth);
 
 				const int32 CachedShiftMouseMovePixelPerDelta = ShiftMouseMovePixelPerDelta.Get();
 				if (CachedShiftMouseMovePixelPerDelta > 1 && MouseEvent.IsShiftDown())
@@ -568,7 +568,7 @@ public:
 
 				if (MouseEvent.IsControlDown())
 				{
-					float DeltaToAdd = MouseEvent.GetCursorDelta().X / SliderWidthInSlateUnits;
+					float DeltaToAdd = (float)MouseEvent.GetCursorDelta().X / SliderWidthInSlateUnits;
 
 					if (SupportDynamicSliderMaxValue.Get() && (NumericType)InternalValue == GetMaxSliderValue())
 					{
@@ -605,7 +605,7 @@ public:
 					FractionFilled *= SliderWidthInSlateUnits;
 
 					// Now add the delta to the fraction filled, this causes the spin.
-					FractionFilled += MouseEvent.GetCursorDelta().X;
+					FractionFilled += (float)MouseEvent.GetCursorDelta().X;
 
 					// Clamp the fraction to be within the bounds of the geometry.
 					FractionFilled = FMath::Clamp(FractionFilled, 0.0f, SliderWidthInSlateUnits);
@@ -692,9 +692,9 @@ public:
 			// First SetText is to update the value before calling CommitValue. Otherwise, when the text lose
 			// focus from the CommitValue, it will override the value we just committed.
 			// The second SetText is to update the text to the InternalValue since it could have been clamped.
-			EditableText->SetText(FText::FromString(Interface->ToString(NewValue)));
+			EditableText->SetText(FText::FromString(Interface->ToString((NumericType)NewValue)));
 			CommitValue(RoundedNewValue, NewValue, CommittedViaSpin, ETextCommit::OnEnter);
-			EditableText->SetText(FText::FromString(Interface->ToString(InternalValue)));
+			EditableText->SetText(FText::FromString(Interface->ToString((NumericType)InternalValue)));
 
 			return FReply::Handled();
 		}
@@ -1019,11 +1019,11 @@ protected:
 		// Update the max slider value based on the current value if we're in dynamic mode
 		if (SupportDynamicSliderMaxValue.Get() && ValueAttribute.Get() > GetMaxSliderValue())
 		{
-			ApplySliderMaxValueChanged(ValueAttribute.Get() - GetMaxSliderValue(), true);
+			ApplySliderMaxValueChanged(float(ValueAttribute.Get() - GetMaxSliderValue()), true);
 		}
 		else if (SupportDynamicSliderMinValue.Get() && ValueAttribute.Get() < GetMinSliderValue())
 		{
-			ApplySliderMinValueChanged(ValueAttribute.Get() - GetMinSliderValue(), true);
+			ApplySliderMinValueChanged(float(ValueAttribute.Get() - GetMinSliderValue()), true);
 		}
 
 		if (CommitMethod == CommittedViaTypeIn || CommitMethod == CommittedViaArrowKey)
@@ -1137,7 +1137,7 @@ private:
 	TAttribute<float> MinDesiredWidth;
 	float GetTextMinDesiredWidth() const
 	{
-		return FMath::Max(0.0f, MinDesiredWidth.Get() - Style->ArrowsImage.ImageSize.X);
+		return FMath::Max(0.0f, MinDesiredWidth.Get() - (float)Style->ArrowsImage.ImageSize.X);
 	}
 
 	/** Check whether a typed character is valid */
diff --git a/Engine/Source/Runtime/Slate/Public/Widgets/Views/STileView.h b/Engine/Source/Runtime/Slate/Public/Widgets/Views/STileView.h
--- a/Engine/Source/Runtime/Slate/Public/Widgets/Views/STileView.h
+++ b/Engine/Source/Runtime/Slate/Public/Widgets/Views/STileView.h
@@ -293,7 +293,7 @@ public:
 			FTableViewDimensions DimensionsUsedSoFar(this->Orientation);
 			
 			// Index of the item at which we start generating based on how far scrolled down we are
-			int32 StartIndex = FMath::Max( 0, FMath::FloorToInt(ClampedScrollOffset / NumItemsPerLine) * NumItemsPerLine);
+			int32 StartIndex = FMath::Max( 0, FMath::FloorToInt32(ClampedScrollOffset / NumItemsPerLine) * NumItemsPerLine);
 
 			// Let the WidgetGenerator know that we are starting a pass so that it can keep track of data items and widgets.
 			this->WidgetGenerator.OnBeginGenerationPass();
@@ -316,7 +316,7 @@ public:
 					if (bFirstLine)
 					{
 						bFirstLine = false;
-						LineFraction -= FMath::Fractional(ClampedScrollOffset / NumItemsPerLine);
+						LineFraction -= (float)FMath::Fractional(ClampedScrollOffset / NumItemsPerLine);
 					}
 
 					DimensionsUsedSoFar.ScrollAxis += TileDimensions.ScrollAxis * LineFraction;
@@ -356,7 +356,7 @@ public:
 			// We have completed the generation pass. The WidgetGenerator will clean up unused Widgets.
 			this->WidgetGenerator.OnEndGenerationPass();
 
-			const float TotalGeneratedLineAxisSize = FMath::CeilToFloat(NumLinesShownOnScreen) * TileDimensions.ScrollAxis;
+			const float TotalGeneratedLineAxisSize = (float)(FMath::CeilToFloat(NumLinesShownOnScreen) * TileDimensions.ScrollAxis);
 			return STableViewBase::FReGenerateResults(ClampedScrollOffset, TotalGeneratedLineAxisSize, NumLinesShownOnScreen, bIsAtEndOfList && !bHasFilledAvailableArea);
 		}
 
@@ -406,7 +406,7 @@ protected:
 		{
 			const double NewScrollOffset = this->DesiredScrollOffset + ((ScrollByAmountInSlateUnits * GetNumItemsPerLine()) / GetTileDimensions().ScrollAxis);
 
-			return this->ScrollTo( NewScrollOffset );
+			return this->ScrollTo( (float)NewScrollOffset );
 		}
 
 		return 0.f;
@@ -435,7 +435,7 @@ protected:
 			{
 				const float NumLinesInView = FTableViewDimensions(this->Orientation, ListViewGeometry.GetLocalSize()).ScrollAxis / GetTileDimensions().ScrollAxis;
 
-				float NumLiveWidgets = this->GetNumLiveWidgets();
+				double NumLiveWidgets = this->GetNumLiveWidgets();
 				if (NumLiveWidgets == 0 && this->IsPendingRefresh())
 				{
 					// Use the last number of widgets on screen to estimate if we actually need to scroll.
@@ -455,7 +455,7 @@ protected:
 				const int32 NumItemsPerLine = GetNumItemsPerLine();
 				const double ScrollLineOffset = this->GetTargetScrollOffset() / NumItemsPerLine;
 				const int32 LineOfItem = FMath::FloorToInt((float)IndexOfItem / (float)NumItemsPerLine);
-				const int32 NumFullLinesInView = FMath::FloorToInt(ScrollLineOffset + NumLinesInView) - FMath::CeilToInt(ScrollLineOffset);
+				const int32 NumFullLinesInView = FMath::FloorToInt32(ScrollLineOffset + NumLinesInView) - FMath::CeilToInt32(ScrollLineOffset);
 				
 				const double MinDisplayedLine = this->bNavigateOnScrollIntoView ? FMath::FloorToDouble(ScrollLineOffset) : FMath::CeilToDouble(ScrollLineOffset);
 				const double MaxDisplayedLine = this->bNavigateOnScrollIntoView ? FMath::CeilToDouble(ScrollLineOffset + NumFullLinesInView) : FMath::FloorToDouble(ScrollLineOffset + NumFullLinesInView);
@@ -463,7 +463,7 @@ protected:
 				if (LineOfItem < MinDisplayedLine || LineOfItem > MaxDisplayedLine)
 				{
 					// Set the line with the item at the beginning of the view area
-					float NewLineOffset = LineOfItem;
+					float NewLineOffset = (float)LineOfItem;
 					// Center the line in the view area
 					NewLineOffset -= NumLinesInView * 0.5f;
 					// Convert the line offset into an item offset
@@ -471,7 +471,7 @@ protected:
 					// And clamp the scroll offset within the allowed limits
 					NewScrollOffset = FMath::Clamp(NewScrollOffset, 0., (double)(GetNumItemsBeingObserved() - NumItemsPerLine * NumLinesInView));
 
-					this->SetScrollOffset(NewScrollOffset);
+					this->SetScrollOffset((float)NewScrollOffset);
 				}
 				else if (this->bNavigateOnScrollIntoView)
 				{
