Add support for ascent/descent override in Font.
Bonus: added the ability to scale the fonts in the preview from the font editor, for clearer preview.

#tests Using the font provided in the jira and the new settable values, ensured that I was able to get result similar to a 'standard' font, on single line text, multi line text, rich text combining the problematic font and another one, as well as using a composite font mixing the fonts. Tested with different size. Validated that modifying the new font settings was refreshing the font displayed. Replaced in a game project the font "burbank" by the new font and got it properly vertically aligned using new settings.

diff --git a/Engine/Source/Editor/FontEditor/Private/FontEditor.cpp b/Engine/Source/Editor/FontEditor/Private/FontEditor.cpp
--- a/Engine/Source/Editor/FontEditor/Private/FontEditor.cpp
+++ b/Engine/Source/Editor/FontEditor/Private/FontEditor.cpp
@@ -71,6 +71,7 @@
 #include "Widgets/Docking/SDockTab.h"
 #include "Widgets/Input/SCheckBox.h"
 #include "Widgets/Input/SEditableTextBox.h"
+#include "Widgets/Input/SNumericEntryBox.h"
 #include "Widgets/SBoxPanel.h"
 #include "Widgets/Text/STextBlock.h"
 
@@ -433,6 +434,16 @@ void FFontEditor::OnPreviewTextChanged(const FText& Text)
 	FontPreviewWidget->SetPreviewText(Text);
 }
 
+TOptional<float> FFontEditor::GetDrawFontScale() const
+{
+	return FontPreviewWidget->GetPreviewFontScale();
+}
+
+void FFontEditor::OnDrawFontScaleChanged(float InNewValue, ETextCommit::Type CommitType)
+{
+	FontPreviewWidget->SetPreviewFontScale(InNewValue);
+}
+
 ECheckBoxState FFontEditor::GetDrawFontMetricsState() const
 {
 	return (FontPreviewWidget->GetPreviewFontMetrics()) ? ECheckBoxState::Checked : ECheckBoxState::Unchecked;
@@ -641,6 +652,17 @@ void FFontEditor::CreateInternalWidgets()
 			.OnTextChanged(this, &FFontEditor::OnPreviewTextChanged)
 		]
 
+		+ SHorizontalBox::Slot()
+		.AutoWidth()
+		[
+			SNew(SNumericEntryBox<float>)
+			.Value(this, &FFontEditor::GetDrawFontScale)
+			.MinValue(1.f)
+			.MaxValue(10.f)
+			.OnValueCommitted(this, &FFontEditor::OnDrawFontScaleChanged)
+		]
+
+
 		+SHorizontalBox::Slot()
 		.AutoWidth()
 		.Padding(FMargin(2.0f, 0.0f, 0.0f, 0.0f))
diff --git a/Engine/Source/Editor/FontEditor/Private/FontEditor.h b/Engine/Source/Editor/FontEditor/Private/FontEditor.h
--- a/Engine/Source/Editor/FontEditor/Private/FontEditor.h
+++ b/Engine/Source/Editor/FontEditor/Private/FontEditor.h
@@ -79,6 +79,10 @@ protected:
 	ECheckBoxState GetDrawFontMetricsState() const;
 	void OnDrawFontMetricsStateChanged(ECheckBoxState NewState);
 
+	/** Called to handle the "Draw Font Scale" numeric entry box */
+	TOptional<float> GetDrawFontScale() const;
+	void OnDrawFontScaleChanged(float InNewValue, ETextCommit::Type CommitType);
+
 	//~ Begin FEditorUndoClient Interface
 	/** Handles any post undo cleanup of the GUI so that we don't have stale data being displayed. */
 	virtual void PostUndo(bool bSuccess) override;
diff --git a/Engine/Source/Editor/FontEditor/Private/SFontEditorViewport.cpp b/Engine/Source/Editor/FontEditor/Private/SFontEditorViewport.cpp
--- a/Engine/Source/Editor/FontEditor/Private/SFontEditorViewport.cpp
+++ b/Engine/Source/Editor/FontEditor/Private/SFontEditorViewport.cpp
@@ -83,6 +83,8 @@ public:
 	const FColor& GetBackgroundColor() const;
 	void SetForegroundColor(const FColor& InForgroundColor);
 	const FColor& GetForegroundColor() const;
+	void SetPreviewFontScale(float InScale);
+	float GetPreviewFontScale() const;
 	void SetDrawFontMetrics(const bool InDrawFontMetrics);
 	bool GetDrawFontMetrics() const;
 	
@@ -125,6 +127,8 @@ private:
 	/** Should we draw the font metrics in the preview? */
 	bool bDrawFontMetrics;
 
+	float DrawFontScale;
+
 	/** The size of the gap between pages */
 	const int32 PageGap;
 };
@@ -132,6 +136,7 @@ private:
 FFontEditorViewportClient::FFontEditorViewportClient(TWeakPtr<SFontEditorViewport> InFontEditorViewport)
 	: FontEditorViewportPtr(InFontEditorViewport)
 	, CurrentSelectedPage(INDEX_NONE)
+	, DrawFontScale(1.0f)
 	, PageGap(4)
 {
 	PreviewText = LOCTEXT("DefaultPreviewText", "The quick brown fox jumps over the lazy dog");
@@ -250,7 +255,8 @@ void FFontEditorViewportClient::Draw(FViewport* Viewport, FCanvas* Canvas)
 		// And draw the text with the foreground color
 		if (Font->FontCacheType == EFontCacheType::Runtime)
 		{
-			const float FontScale = GetDPIScale();
+			const float FontScale = GetDPIScale() * DrawFontScale;
+			const float DebugDrawScale = GetDPIScale();
 
 			TSharedRef<FSlateFontCache> FontCache = FSlateApplication::Get().GetRenderer()->GetFontCache();
 
@@ -390,7 +396,7 @@ void FFontEditorViewportClient::Draw(FViewport* Viewport, FCanvas* Canvas)
 			if (bDrawFontMetrics)
 			{
 				const FSlateFontInfo FontInfo = FAppStyle::GetFontStyle("NormalFont");
-				const float KeyBoxSize = 14.0f * FontScale;
+				const float KeyBoxSize = 14.0f * DebugDrawScale;
 
 				struct FKeyDataType
 				{
@@ -412,13 +418,13 @@ void FFontEditorViewportClient::Draw(FViewport* Viewport, FCanvas* Canvas)
 					KeyBox.LineThickness = KeyBoxSize * 0.5f;
 					Canvas->DrawItem(KeyBox);
 
-					CurPos.X += KeyBoxSize + (4.0f * FontScale);
+					CurPos.X += KeyBoxSize + (4.0f * DebugDrawScale);
 
-					FShapedGlyphSequenceRef KeyLabelShapedText = FontCache->ShapeBidirectionalText(*KeyData.KeyText.ToString(), FontInfo, FontScale, TextBiDi::ETextDirection::LeftToRight, GetDefaultTextShapingMethod());
+					FShapedGlyphSequenceRef KeyLabelShapedText = FontCache->ShapeBidirectionalText(*KeyData.KeyText.ToString(), FontInfo, DebugDrawScale, TextBiDi::ETextDirection::LeftToRight, GetDefaultTextShapingMethod());
 					FCanvasShapedTextItem ShapedTextItem(CurPos, KeyLabelShapedText, FLinearColor(ForegroundColor));
 					Canvas->DrawItem(ShapedTextItem);
 
-					CurPos.X += KeyLabelShapedText->GetMeasuredWidth() + (8.0f * FontScale);
+					CurPos.X += KeyLabelShapedText->GetMeasuredWidth() + (8.0f * DebugDrawScale);
 				}
 			}
 		}
@@ -555,6 +561,16 @@ const FColor& FFontEditorViewportClient::GetForegroundColor() const
 	return ForegroundColor;
 }
 
+void FFontEditorViewportClient::SetPreviewFontScale(float InScale)
+{
+	DrawFontScale = InScale;
+}
+
+float FFontEditorViewportClient::GetPreviewFontScale() const
+{
+	return DrawFontScale;
+}
+
 void FFontEditorViewportClient::SetDrawFontMetrics(const bool InDrawFontMetrics)
 {
 	bDrawFontMetrics = InDrawFontMetrics;
@@ -881,6 +897,26 @@ const FColor& SFontEditorViewport::GetPreviewForegroundColor() const
 	return FColor::White;
 }
 
+void SFontEditorViewport::SetPreviewFontScale(float InScale)
+{
+	if (ViewportClient.IsValid())
+	{
+		ViewportClient->SetPreviewFontScale(InScale);
+
+		RefreshViewport();
+	}
+}
+
+float SFontEditorViewport::GetPreviewFontScale() const
+{
+	if (ViewportClient.IsValid())
+	{
+		return ViewportClient->GetPreviewFontScale();
+	}
+
+	return 1.0f;
+}
+
 void SFontEditorViewport::SetPreviewFontMetrics(const bool InDrawFontMetrics)
 {
 	if (ViewportClient.IsValid())
diff --git a/Engine/Source/Editor/FontEditor/Private/SFontEditorViewport.h b/Engine/Source/Editor/FontEditor/Private/SFontEditorViewport.h
--- a/Engine/Source/Editor/FontEditor/Private/SFontEditorViewport.h
+++ b/Engine/Source/Editor/FontEditor/Private/SFontEditorViewport.h
@@ -47,6 +47,8 @@ public:
 	const FColor& GetPreviewBackgroundColor() const;
 	void SetPreviewForegroundColor(const FColor& ForgroundColor);
 	const FColor& GetPreviewForegroundColor() const;
+	void SetPreviewFontScale(float InScale);
+	float GetPreviewFontScale() const;
 	void SetPreviewFontMetrics(const bool InDrawFontMetrics);
 	bool GetPreviewFontMetrics() const;
 	TWeakPtr<IFontEditor> GetFontEditor() const;
diff --git a/Engine/Source/Runtime/Engine/Classes/Engine/FontFace.h b/Engine/Source/Runtime/Engine/Classes/Engine/FontFace.h
--- a/Engine/Source/Runtime/Engine/Classes/Engine/FontFace.h
+++ b/Engine/Source/Runtime/Engine/Classes/Engine/FontFace.h
@@ -46,6 +46,10 @@ public:
 	virtual EFontHinting GetHinting() const override;
 	virtual EFontLoadingPolicy GetLoadingPolicy() const override;
 	virtual EFontLayoutMethod GetLayoutMethod() const override;
+	virtual bool IsAscendOverridden() const override;
+	virtual int32 GetAscendOverriddenValue() const override;
+	virtual bool IsDescendOverridden() const override;
+	virtual int32 GetDescendOverriddenValue() const override;
 	virtual FFontFaceDataConstRef GetFontFaceData() const override;
 	//~ End IFontFaceInterface interface
 
@@ -75,6 +79,22 @@ public:
 	UPROPERTY(EditAnywhere, BlueprintReadWrite, Category=FontFace, AdvancedDisplay)
 	EFontLayoutMethod LayoutMethod;
 
+	/** The typographic ascender of the face, expressed in font units. */
+	UPROPERTY(EditAnywhere, Category=FontFace, AdvancedDisplay, meta = (EditCondition = bIsAscendOverridden, EditConditionHides, ClampMin = "-100", ClampMax = "100"))
+	int32 AscendOverriddenValue;
+
+	/** Activate this option to use the specified ascend value instead of the value from the font. */
+	UPROPERTY(EditAnywhere, Category=FontFace, AdvancedDisplay)
+	bool bIsAscendOverridden;
+
+	/** The typographic ascender of the face, expressed in font units. */
+	UPROPERTY(EditAnywhere, Category=FontFace, AdvancedDisplay, meta = (EditCondition = bIsDescendOverridden, EditConditionHides, ClampMin = "-100", ClampMax = "100"))
+	int32 DescendOverriddenValue;
+
+	/** Activate this option to use the specified descend value instead of the value from the font. */
+	UPROPERTY(EditAnywhere, Category=FontFace, AdvancedDisplay)
+	bool bIsDescendOverridden;
+
 	/** The data associated with the font face. This should always be filled in providing the source filename is valid. CacheSubFaces should be called after manually changing this property. */
 	FFontFaceDataRef FontFaceData;
 
diff --git a/Engine/Source/Runtime/Engine/Private/FontFace.cpp b/Engine/Source/Runtime/Engine/Private/FontFace.cpp
--- a/Engine/Source/Runtime/Engine/Private/FontFace.cpp
+++ b/Engine/Source/Runtime/Engine/Private/FontFace.cpp
@@ -17,7 +17,12 @@
 DEFINE_LOG_CATEGORY_STATIC(LogFontFace, Log, All);
 
 UFontFace::UFontFace()
-	: FontFaceData(FFontFaceData::MakeFontFaceData())
+	: AscendOverriddenValue(0)
+	, bIsAscendOverridden(false)
+	, DescendOverriddenValue(0)
+	, bIsDescendOverridden(false)
+	, FontFaceData(FFontFaceData::MakeFontFaceData())
+
 {
 }
 
@@ -203,6 +208,26 @@ EFontLayoutMethod UFontFace::GetLayoutMethod() const
 	return LayoutMethod;
 }
 
+bool UFontFace::IsAscendOverridden() const
+{
+	return bIsAscendOverridden;
+}
+
+int32 UFontFace::GetAscendOverriddenValue() const
+{
+	return AscendOverriddenValue;
+}
+
+bool UFontFace::IsDescendOverridden() const
+{
+	return bIsDescendOverridden;
+}
+
+int32 UFontFace::GetDescendOverriddenValue() const
+{
+	return DescendOverriddenValue;
+}
+
 FFontFaceDataConstRef UFontFace::GetFontFaceData() const
 {
 	return FontFaceData;
diff --git a/Engine/Source/Runtime/SlateCore/Private/Fonts/CompositeFont.cpp b/Engine/Source/Runtime/SlateCore/Private/Fonts/CompositeFont.cpp
--- a/Engine/Source/Runtime/SlateCore/Private/Fonts/CompositeFont.cpp
+++ b/Engine/Source/Runtime/SlateCore/Private/Fonts/CompositeFont.cpp
@@ -127,6 +127,46 @@ EFontLayoutMethod FFontData::GetLayoutMethod() const
 	return EFontLayoutMethod::Metrics;
 }
 
+bool FFontData::IsAscendOverridden() const
+{
+	if (FontFaceAsset)
+	{
+		const IFontFaceInterface* FontFace = CastChecked<const IFontFaceInterface>(FontFaceAsset);
+		return FontFace->IsAscendOverridden();
+	}
+	return false;
+}
+
+int32 FFontData::GetAscendOverriddenValue() const
+{
+	if (FontFaceAsset)
+	{
+		const IFontFaceInterface* FontFace = CastChecked<const IFontFaceInterface>(FontFaceAsset);
+		return FontFace->GetAscendOverriddenValue();
+	}
+	return 0;
+}
+
+bool FFontData::IsDescendOverridden() const
+{
+	if (FontFaceAsset)
+	{
+		const IFontFaceInterface* FontFace = CastChecked<const IFontFaceInterface>(FontFaceAsset);
+		return FontFace->IsDescendOverridden();
+	}
+	return false;
+}
+
+int32 FFontData::GetDescendOverriddenValue() const
+{
+	if (FontFaceAsset)
+	{
+		const IFontFaceInterface* FontFace = CastChecked<const IFontFaceInterface>(FontFaceAsset);
+		return FontFace->GetDescendOverriddenValue();
+	}
+	return 0;
+}
+
 FFontFaceDataConstPtr FFontData::GetFontFaceData() const
 {
 	if (FontFaceAsset)
diff --git a/Engine/Source/Runtime/SlateCore/Private/Fonts/FontCacheCompositeFont.cpp b/Engine/Source/Runtime/SlateCore/Private/Fonts/FontCacheCompositeFont.cpp
--- a/Engine/Source/Runtime/SlateCore/Private/Fonts/FontCacheCompositeFont.cpp
+++ b/Engine/Source/Runtime/SlateCore/Private/Fonts/FontCacheCompositeFont.cpp
@@ -448,12 +448,19 @@ TSharedPtr<FFreeTypeFace> FCompositeFontCache::GetFontFace(const FFontData& InFo
 		// Got a valid font?
 		if (FaceAndMemory.IsValid())
 		{
+			FaceAndMemory->OverrideAscend(InFontData.IsAscendOverridden(), InFontData.GetAscendOverriddenValue());
+			FaceAndMemory->OverrideDescend(InFontData.IsDescendOverridden(), InFontData.GetDescendOverriddenValue());
+
 			FontFaceMap.Add(InFontData, FaceAndMemory);
 
 			if (!LoadLogMessage.IsEmpty())
 			{
 				UE_LOG(LogSlate, Log, TEXT("%s"), *LoadLogMessage);
 			}
+
+#if WITH_EDITOR //Triggers texts layout regeneration only in editor, don't want to slow down things in engine.
+			GSlateLayoutGeneration++;
+#endif
 		}
 		else
 		{
diff --git a/Engine/Source/Runtime/SlateCore/Private/Fonts/FontCacheFreeType.h b/Engine/Source/Runtime/SlateCore/Private/Fonts/FontCacheFreeType.h
--- a/Engine/Source/Runtime/SlateCore/Private/Fonts/FontCacheFreeType.h
+++ b/Engine/Source/Runtime/SlateCore/Private/Fonts/FontCacheFreeType.h
@@ -222,16 +222,28 @@ public:
 
 	FORCEINLINE FT_Pos GetScaledHeight() const
 	{
+		if (IsAscendOverridden || IsDescendOverridden)
+			return GetAscender() - GetDescender();
 		return FreeTypeUtils::GetScaledHeight(FTFace, LayoutMethod);
 	}
 
 	FORCEINLINE FT_Pos GetAscender() const
 	{
+		if (IsAscendOverridden)
+		{
+			FT_F26Dot6 ScaledAscender = FT_MulFix(AscendOverrideValue, FTFace->size->metrics.y_scale);
+			return (ScaledAscender + 0b111111) & ~0b111111; //(26.6 fixed point ceil). Using ceiling of scaled ascend, as recommended by Freetype, to avoid grid fitting/hinting issues.
+		}
 		return FreeTypeUtils::GetAscender(FTFace, LayoutMethod);
 	}
 
 	FORCEINLINE FT_Pos GetDescender() const
 	{
+		if (IsDescendOverridden)
+		{
+			FT_F26Dot6 ScaledDescender =  FT_MulFix(DescendOverrideValue, FTFace->size->metrics.y_scale);
+			return ScaledDescender & ~0b111111; //(26.6 fixed point floor). Using floor of scaled descend, as recommended by Freetype, to avoid grid fitting/hinting issues.
+		}
 		return FreeTypeUtils::GetDescender(FTFace, LayoutMethod);
 	}
 
@@ -268,6 +280,22 @@ public:
 #endif
 	}
 
+	void OverrideAscend(bool InOverride, int32 Value = 0)
+	{
+#if WITH_FREETYPE
+		IsAscendOverridden = InOverride;
+		AscendOverrideValue = FreeTypeUtils::ConvertPixelTo26Dot6<FT_F26Dot6>(Value);
+#endif //WITH_FREETYPE
+	}
+
+	void OverrideDescend(bool InOverride, int32 Value = 0)
+	{
+#if WITH_FREETYPE
+		IsDescendOverridden = InOverride;
+		DescendOverrideValue = FreeTypeUtils::ConvertPixelTo26Dot6<FT_F26Dot6>(Value);
+#endif //WITH_FREETYPE
+	}
+
 	void FailAsyncLoad();
 	void CompleteAsyncLoad(const FFreeTypeLibrary* InFTLibrary, FFontFaceDataConstRef InMemory, const int32 InFaceIndex);
 
@@ -311,6 +339,11 @@ private:
 	FFTStreamHandler FTStreamHandler;
 	FT_StreamRec FTStream;
 	FT_Open_Args FTFaceOpenArgs;
+
+	bool IsAscendOverridden = false;
+	bool IsDescendOverridden = false;
+	FT_F26Dot6 AscendOverrideValue = 0;
+	FT_F26Dot6 DescendOverrideValue = 0;
 #endif // WITH_FREETYPE
 
 	TSet<FName> Attributes;
diff --git a/Engine/Source/Runtime/SlateCore/Public/Fonts/CompositeFont.h b/Engine/Source/Runtime/SlateCore/Public/Fonts/CompositeFont.h
--- a/Engine/Source/Runtime/SlateCore/Public/Fonts/CompositeFont.h
+++ b/Engine/Source/Runtime/SlateCore/Public/Fonts/CompositeFont.h
@@ -167,6 +167,18 @@ struct SLATECORE_API FFontData
 	/** Get the method to use when laying out the font? */
 	EFontLayoutMethod GetLayoutMethod() const;
 
+	/** Returns true if the ascend is overridden. */
+	bool IsAscendOverridden() const;
+
+	/** Returns the overridden value of the ascend. This value will be used only if IsAscendOverridden returns true. */
+	int32 GetAscendOverriddenValue() const;
+
+	/** Returns true if the descend is overridden. */
+	bool IsDescendOverridden() const;
+
+	/** Returns the overridden value of the descend. This value will be used only if IsDescendOverridden returns true. */
+	int32 GetDescendOverriddenValue() const;
+
 	/** Get the data buffer containing the data for the current font face. */
 	FFontFaceDataConstPtr GetFontFaceData() const;
 
diff --git a/Engine/Source/Runtime/SlateCore/Public/Fonts/FontFaceInterface.h b/Engine/Source/Runtime/SlateCore/Public/Fonts/FontFaceInterface.h
--- a/Engine/Source/Runtime/SlateCore/Public/Fonts/FontFaceInterface.h
+++ b/Engine/Source/Runtime/SlateCore/Public/Fonts/FontFaceInterface.h
@@ -41,6 +41,18 @@ class IFontFaceInterface
 	/** Get the method to use when laying out the font? */
 	virtual EFontLayoutMethod GetLayoutMethod() const = 0;
 
+	/** Returns true if the ascend is overridden. */
+	virtual bool IsAscendOverridden() const = 0;
+
+	/** Returns the overridden value of the ascend. This value will be used only if IsAscendOverridden returns true. */
+	virtual int32 GetAscendOverriddenValue() const = 0;
+
+	/** Returns true if the descend is overridden. */
+	virtual bool IsDescendOverridden() const = 0;
+
+	/** Returns the overridden value of the descend. This value will be used only if IsDescendOverridden returns true. */
+	virtual int32 GetDescendOverriddenValue() const = 0;
+
 	/** Get the data buffer containing the data for the current font face. */
 	virtual FFontFaceDataConstRef GetFontFaceData() const = 0;
 };
