Slate SDF fonts - separate resolution for multi-channel and single-channel distance fields
#tests Tested in editor
#rnx <Release Notes exclude>

diff --git a/Engine/Source/Editor/FontEditor/Private/FontFaceDetailsCustomization.cpp b/Engine/Source/Editor/FontEditor/Private/FontFaceDetailsCustomization.cpp
--- a/Engine/Source/Editor/FontEditor/Private/FontFaceDetailsCustomization.cpp
+++ b/Engine/Source/Editor/FontEditor/Private/FontFaceDetailsCustomization.cpp
@@ -15,6 +15,7 @@
 #include "Framework/Application/SlateApplication.h"
 #include "HAL/Platform.h"
 #include "IDesktopPlatform.h"
+#include "IDetailGroup.h"
 #include "Internationalization/Internationalization.h"
 #include "Layout/Margin.h"
 #include "Misc/AssertionMacros.h"
@@ -94,13 +95,45 @@ void FFontFaceDetailsCustomization::CustomizeDetails(IDetailLayoutBuilder& Detai
 			];
 	}
 
-	// Hide distance field settings if disabled by CVar
-	if (!IsSlateSdfTextFeatureEnabled())
+	// Distance field properties
+	if (IsSlateSdfTextFeatureEnabled())
 	{
+		bool bShowDistanceFieldGroups = true;
+		if (ObjectsBeingEdited.Num() == 1)
+		{
+			const UFontFace* FontFace = Cast<UFontFace>(ObjectsBeingEdited[0].Get());
+			if (FontFace && !FontFace->bEnableDistanceFieldRendering)
+			{
+				bShowDistanceFieldGroups = false;
+			}
+		}
+		if (bShowDistanceFieldGroups)
+		{
+			IDetailCategoryBuilder& DistanceFieldModeCategory = DetailBuilder.EditCategory("DistanceFieldMode");
+			// Make sure that bEnableDistanceFieldRendering is listed first
+			DistanceFieldModeCategory.AddProperty(DetailBuilder.GetProperty(GET_MEMBER_NAME_CHECKED(UFontFace, bEnableDistanceFieldRendering)));
+			// Single-channel distance field resolution group
+			IDetailGroup& SingleChannelResolutionGroup = DistanceFieldModeCategory.AddGroup(FName(TEXT("SingleChannelDistanceFieldResolution")), LOCTEXT("FontFaceSingleChannelDistanceFieldResolutionGroup", "Single-Channel Distance Field Resolution (px/em)"), false, true);
+			SingleChannelResolutionGroup.AddPropertyRow(DetailBuilder.GetProperty(GET_MEMBER_NAME_CHECKED(UFontFace, MinDistanceFieldPpem)));
+			SingleChannelResolutionGroup.AddPropertyRow(DetailBuilder.GetProperty(GET_MEMBER_NAME_CHECKED(UFontFace, MidDistanceFieldPpem)));
+			SingleChannelResolutionGroup.AddPropertyRow(DetailBuilder.GetProperty(GET_MEMBER_NAME_CHECKED(UFontFace, MaxDistanceFieldPpem)));
+			// Multi-channel distance field resolution group
+			IDetailGroup& MultiChannelResolutionGroup = DistanceFieldModeCategory.AddGroup(FName(TEXT("MultiChannelDistanceFieldResolution")), LOCTEXT("FontFaceMultiChannelDistanceFieldResolutionGroup", "Multi-Channel Distance Field Resolution (px/em)"), false, true);
+			MultiChannelResolutionGroup.AddPropertyRow(DetailBuilder.GetProperty(GET_MEMBER_NAME_CHECKED(UFontFace, MinMultiDistanceFieldPpem)));
+			MultiChannelResolutionGroup.AddPropertyRow(DetailBuilder.GetProperty(GET_MEMBER_NAME_CHECKED(UFontFace, MidMultiDistanceFieldPpem)));
+			MultiChannelResolutionGroup.AddPropertyRow(DetailBuilder.GetProperty(GET_MEMBER_NAME_CHECKED(UFontFace, MaxMultiDistanceFieldPpem)));
+		}
+	}
+	else
+	{
+		// Hide distance field settings if disabled by CVar
 		DetailBuilder.GetProperty(GET_MEMBER_NAME_CHECKED(UFontFace, bEnableDistanceFieldRendering))->MarkHiddenByCustomization();
 		DetailBuilder.GetProperty(GET_MEMBER_NAME_CHECKED(UFontFace, MinDistanceFieldPpem))->MarkHiddenByCustomization();
 		DetailBuilder.GetProperty(GET_MEMBER_NAME_CHECKED(UFontFace, MidDistanceFieldPpem))->MarkHiddenByCustomization();
 		DetailBuilder.GetProperty(GET_MEMBER_NAME_CHECKED(UFontFace, MaxDistanceFieldPpem))->MarkHiddenByCustomization();
+		DetailBuilder.GetProperty(GET_MEMBER_NAME_CHECKED(UFontFace, MinMultiDistanceFieldPpem))->MarkHiddenByCustomization();
+		DetailBuilder.GetProperty(GET_MEMBER_NAME_CHECKED(UFontFace, MidMultiDistanceFieldPpem))->MarkHiddenByCustomization();
+		DetailBuilder.GetProperty(GET_MEMBER_NAME_CHECKED(UFontFace, MaxMultiDistanceFieldPpem))->MarkHiddenByCustomization();
 		DetailBuilder.GetProperty(GET_MEMBER_NAME_CHECKED(UFontFace, PlatformRasterizationModeOverrides))->MarkHiddenByCustomization();
 	}
 }
diff --git a/Engine/Source/Runtime/Engine/Classes/Engine/FontFace.h b/Engine/Source/Runtime/Engine/Classes/Engine/FontFace.h
--- a/Engine/Source/Runtime/Engine/Classes/Engine/FontFace.h
+++ b/Engine/Source/Runtime/Engine/Classes/Engine/FontFace.h
@@ -144,20 +144,32 @@ public:
 	UPROPERTY(EditAnywhere, Category=DistanceFieldMode)
 	bool bEnableDistanceFieldRendering = false;
 	
-	/** Distance field px/em resolution "low" quality value */
-	UPROPERTY(EditAnywhere, Category=DistanceFieldMode, meta=(ClampMin=8, ClampMax=256, DisplayName="Low SDF resolution (px/em)"))
+	/** Single-channel distance field px/em resolution "low" quality value */
+	UPROPERTY(EditAnywhere, Category=DistanceFieldMode, meta=(ClampMin=8, ClampMax=256, DisplayName="Low Quality", EditCondition="bEnableDistanceFieldRendering", EditConditionHides))
 	int32 MinDistanceFieldPpem = 32;
 
-	/** Distance field px/em resolution "medium" quality value */
-	UPROPERTY(EditAnywhere, Category=DistanceFieldMode, meta=(ClampMin=8, ClampMax=256, DisplayName="Medium SDF resolution (px/em)"))
+	/** Single-channel distance field px/em resolution "medium" quality value */
+	UPROPERTY(EditAnywhere, Category=DistanceFieldMode, meta=(ClampMin=8, ClampMax=256, DisplayName="Medium Quality", EditCondition="bEnableDistanceFieldRendering", EditConditionHides))
 	int32 MidDistanceFieldPpem = 48;
 
-	/** Distance field px/em resolution "high" quality value */
-	UPROPERTY(EditAnywhere, Category=DistanceFieldMode, meta=(ClampMin=8, ClampMax=256, DisplayName="High SDF resolution (px/em)"))
+	/** Single-channel distance field px/em resolution "high" quality value */
+	UPROPERTY(EditAnywhere, Category=DistanceFieldMode, meta=(ClampMin=8, ClampMax=256, DisplayName="High Quality", EditCondition="bEnableDistanceFieldRendering", EditConditionHides))
 	int32 MaxDistanceFieldPpem = 64;
+	
+	/** Multi-channel distance field px/em resolution "low" quality value */
+	UPROPERTY(EditAnywhere, Category=DistanceFieldMode, meta=(ClampMin=8, ClampMax=256, DisplayName="Low Quality", EditCondition="bEnableDistanceFieldRendering", EditConditionHides))
+	int32 MinMultiDistanceFieldPpem = 32;
+
+	/** Multi-channel distance field px/em resolution "medium" quality value */
+	UPROPERTY(EditAnywhere, Category=DistanceFieldMode, meta=(ClampMin=8, ClampMax=256, DisplayName="Medium Quality", EditCondition="bEnableDistanceFieldRendering", EditConditionHides))
+	int32 MidMultiDistanceFieldPpem = 40;
+
+	/** Multi-channel distance field px/em resolution "high" quality value */
+	UPROPERTY(EditAnywhere, Category=DistanceFieldMode, meta=(ClampMin=8, ClampMax=256, DisplayName="High Quality", EditCondition="bEnableDistanceFieldRendering", EditConditionHides))
+	int32 MaxMultiDistanceFieldPpem = 56;
 
 	/** If set, allows to override distance field modes set in device profiles */
-	UPROPERTY(EditAnywhere, Category=DistanceFieldMode, meta=(ClampMin=8, ClampMax=256, DisplayName="Override platform rasterization mode"))
+	UPROPERTY(EditAnywhere, Category=DistanceFieldMode, meta=(ClampMin=8, ClampMax=256, DisplayName="Override Platform Rasterization Mode", EditCondition="bEnableDistanceFieldRendering", EditConditionHides))
 	TOptional<FFontFacePlatformRasterizationOverrides> PlatformRasterizationModeOverrides;
 
 private:
diff --git a/Engine/Source/Runtime/Engine/Private/FontFace.cpp b/Engine/Source/Runtime/Engine/Private/FontFace.cpp
--- a/Engine/Source/Runtime/Engine/Private/FontFace.cpp
+++ b/Engine/Source/Runtime/Engine/Private/FontFace.cpp
@@ -326,18 +326,19 @@ void UFontFace::UpdateDeviceRasterizationSettings()
 			}
 		}
 
+		const bool bMultiChannel = DeviceRasterizationSettings.Mode == EFontRasterizationMode::Msdf;
 		const int32 DeviceResolutionLevel = CVarFontFaceDistanceFieldResolutionLevel.GetValueOnAnyThread();
 		if (DeviceResolutionLevel <= 1) // Low
 		{
-			DeviceRasterizationSettings.DistanceFieldPpem = MinDistanceFieldPpem;
+			DeviceRasterizationSettings.DistanceFieldPpem = bMultiChannel ? MinMultiDistanceFieldPpem : MinDistanceFieldPpem;
 		}
 		else if (DeviceResolutionLevel >= 3) // High
 		{
-			DeviceRasterizationSettings.DistanceFieldPpem = MaxDistanceFieldPpem;
+			DeviceRasterizationSettings.DistanceFieldPpem = bMultiChannel ? MaxMultiDistanceFieldPpem : MaxDistanceFieldPpem;
 		}
 		else // 2 = Medium
 		{
-			DeviceRasterizationSettings.DistanceFieldPpem = MidDistanceFieldPpem;
+			DeviceRasterizationSettings.DistanceFieldPpem = bMultiChannel ? MidMultiDistanceFieldPpem : MidDistanceFieldPpem;
 		}
 	}
 }
