Slate SDF fonts - Font Face editor preview
#tests Tested in editor
#rnx <Release Notes exclude>

diff --git a/Engine/Plugins/Editor/EngineAssetDefinitions/Source/Private/AssetDefinition_FontFace.cpp b/Engine/Plugins/Editor/EngineAssetDefinitions/Source/Private/AssetDefinition_FontFace.cpp
--- a/Engine/Plugins/Editor/EngineAssetDefinitions/Source/Private/AssetDefinition_FontFace.cpp
+++ b/Engine/Plugins/Editor/EngineAssetDefinitions/Source/Private/AssetDefinition_FontFace.cpp
@@ -6,6 +6,7 @@
 #include "FontEditorModule.h"
 #include "ToolMenus.h"
 #include "EditorReimportHandler.h"
+#include "FontEditorModule.h"
 #include "ToolMenu.h"
 #include "Toolkits/SimpleAssetEditor.h"
 #include "ToolMenuSection.h"
@@ -21,7 +22,7 @@ EAssetCommandResult UAssetDefinition_FontFace::OpenAssets(const FAssetOpenArgs&
 	// into a single editor instance which doesn't really work for font face assets
 	for (UFontFace* FontFace : OpenArgs.LoadObjects<UFontFace>())
 	{
-		FSimpleAssetEditor::CreateEditor(OpenArgs.GetToolkitMode(), OpenArgs.ToolkitHost, FontFace);
+		FontEditorModule->CreateFontFaceEditor(OpenArgs.GetToolkitMode(), OpenArgs.ToolkitHost, FontFace);
 	}
 
 	return EAssetCommandResult::Handled;
diff --git a/Engine/Source/Editor/FontEditor/Private/FontEditorModule.cpp b/Engine/Source/Editor/FontEditor/Private/FontEditorModule.cpp
--- a/Engine/Source/Editor/FontEditor/Private/FontEditorModule.cpp
+++ b/Engine/Source/Editor/FontEditor/Private/FontEditorModule.cpp
@@ -6,6 +6,7 @@
 #include "Engine/Font.h"
 #include "Engine/FontFace.h"
 #include "FontEditor.h"
+#include "FontFaceEditor.h"
 #include "FontFaceDetailsCustomization.h"
 #include "HAL/Platform.h"
 #include "Modules/ModuleManager.h"
@@ -66,6 +67,14 @@ public:
 		return NewFontEditor;
 	}
 
+	/** Creates a new Font Face editor */
+	virtual TSharedRef<IFontFaceEditor> CreateFontFaceEditor( const EToolkitMode::Type Mode, const TSharedPtr< IToolkitHost >& InitToolkitHost, UFontFace* FontFace ) override
+	{
+		TSharedRef<FFontFaceEditor> NewFontEditor(new FFontFaceEditor());
+		NewFontEditor->InitFontFaceEditor(Mode, InitToolkitHost, FontFace);
+		return NewFontEditor;
+	}
+
 	/** Gets the extensibility managers for outside entities to extend static mesh editor's menus and toolbars */
 	virtual TSharedPtr<FExtensibilityManager> GetMenuExtensibilityManager() override { return MenuExtensibilityManager; }
 	virtual TSharedPtr<FExtensibilityManager> GetToolBarExtensibilityManager() override { return ToolBarExtensibilityManager; }
diff --git a/Engine/Source/Editor/FontEditor/Private/FontFaceEditor.cpp b/Engine/Source/Editor/FontEditor/Private/FontFaceEditor.cpp
new file mode 100644
--- /dev/null
+++ b/Engine/Source/Editor/FontEditor/Private/FontFaceEditor.cpp
@@ -0,0 +1,556 @@
+// Copyright Epic Games, Inc. All Rights Reserved.
+
+#include "FontFaceEditor.h"
+
+#include "Containers/Array.h"
+#include "DetailsViewArgs.h"
+#include "Editor.h"
+#include "Editor/EditorEngine.h"
+#include "EditorReimportHandler.h"
+#include "Engine/Engine.h"
+#include "Engine/EngineTypes.h"
+#include "Engine/Font.h"
+#include "Engine/FontFace.h"
+#include "Engine/UserInterfaceSettings.h"
+#include "FontEditorModule.h"
+#include "Framework/Application/SlateApplication.h"
+#include "Modules/ModuleManager.h"
+#include "PropertyEditorModule.h"
+#include "Styling/AppStyle.h"
+#include "Subsystems/ImportSubsystem.h"
+#include "Textures/SlateIcon.h"
+#include "Toolkits/AssetEditorToolkit.h"
+#include "UObject/ObjectPtr.h"
+#include "Widgets/Docking/SDockTab.h"
+#include "Widgets/Input/SEditableTextBox.h"
+#include "Widgets/Input/SNumericEntryBox.h"
+#include "Widgets/Layout/SGridPanel.h"
+#include "Widgets/Layout/SScrollBox.h"
+#include "Widgets/Text/STextBlock.h"
+
+#define LOCTEXT_NAMESPACE "FontFaceEditor"
+
+DEFINE_LOG_CATEGORY_STATIC(LogFontFaceEditor, Log, All);
+
+const FName FFontFaceEditor::PreviewTabId( TEXT( "FontFaceEditor_FontFacePreview" ) );
+const FName FFontFaceEditor::PropertiesTabId( TEXT( "FontFaceEditor_FontFaceProperties" ) );
+
+void FFontFaceEditor::RegisterTabSpawners(const TSharedRef<class FTabManager>& InTabManager)
+{
+	WorkspaceMenuCategory = InTabManager->AddLocalWorkspaceMenuCategory(LOCTEXT("WorkspaceMenu_FontFaceEditor", "Font Face Editor"));
+	auto WorkspaceMenuCategoryRef = WorkspaceMenuCategory.ToSharedRef();
+
+	FAssetEditorToolkit::RegisterTabSpawners(InTabManager);
+
+	InTabManager->RegisterTabSpawner( PreviewTabId,		FOnSpawnTab::CreateSP(this, &FFontFaceEditor::SpawnTab_Preview) )
+		.SetDisplayName( LOCTEXT("PreviewTab", "Preview") )
+		.SetGroup(WorkspaceMenuCategoryRef)
+		.SetIcon(FSlateIcon(FAppStyle::GetAppStyleSetName(), "FontEditor.Tabs.Preview"));
+
+	InTabManager->RegisterTabSpawner( PropertiesTabId,	FOnSpawnTab::CreateSP(this, &FFontFaceEditor::SpawnTab_Properties) )
+		.SetDisplayName( LOCTEXT("PropertiesTabId", "Details") )
+		.SetGroup(WorkspaceMenuCategoryRef)
+		.SetIcon(FSlateIcon(FAppStyle::GetAppStyleSetName(), "LevelEditor.Tabs.Details"));
+}
+
+void FFontFaceEditor::UnregisterTabSpawners(const TSharedRef<class FTabManager>& InTabManager)
+{
+	FAssetEditorToolkit::UnregisterTabSpawners(InTabManager);
+
+	InTabManager->UnregisterTabSpawner( PreviewTabId );	
+	InTabManager->UnregisterTabSpawner( PropertiesTabId );
+}
+
+FFontFaceEditor::FFontFaceEditor()
+	: FontFace(nullptr)
+{
+}
+
+FFontFaceEditor::~FFontFaceEditor()
+{
+	FReimportManager::Instance()->OnPostReimport().RemoveAll(this);
+
+	if (UEditorEngine* Editor = Cast<UEditorEngine>(GEngine))
+	{
+		Editor->UnregisterForUndo(this);
+		Editor->GetEditorSubsystem<UImportSubsystem>()->OnAssetReimport.RemoveAll(this);
+	}
+}
+
+void FFontFaceEditor::InitFontFaceEditor(const EToolkitMode::Type Mode, const TSharedPtr< class IToolkitHost >& InitToolkitHost, UObject* ObjectToEdit)
+{
+	FReimportManager::Instance()->OnPostReimport().AddRaw(this, &FFontFaceEditor::OnPostReimport);
+
+	// Register to be notified when an object is reimported.
+	GEditor->GetEditorSubsystem<UImportSubsystem>()->OnAssetReimport.AddSP(this, &FFontFaceEditor::OnObjectReimported);
+
+	FCoreUObjectDelegates::OnObjectPropertyChanged.AddSP(this, &FFontFaceEditor::OnObjectPropertyChanged);
+
+	FontFace = CastChecked<UFontFace>(ObjectToEdit);
+
+	// Support undo/redo
+	FontFace->SetFlags(RF_Transactional);
+	
+	if (UEditorEngine* Editor = Cast<UEditorEngine>(GEngine))
+	{
+		Editor->RegisterForUndo(this);
+	}
+
+	CreateInternalWidgets();
+
+	const TSharedRef<FTabManager::FLayout> StandaloneDefaultLayout = FTabManager::NewLayout("Standalone_FontEditor_Layout_v4")
+	->AddArea
+	(
+		FTabManager::NewPrimaryArea() ->SetOrientation( Orient_Vertical )
+		->Split
+		(
+			FTabManager::NewSplitter() ->SetOrientation(Orient_Vertical) ->SetSizeCoefficient(0.65f)
+			->Split
+			(
+				FTabManager::NewStack() ->SetSizeCoefficient(0.85f)
+				->AddTab( PropertiesTabId, ETabState::OpenedTab )
+			)
+			->Split
+			(
+				FTabManager::NewStack() ->SetSizeCoefficient(0.15f)
+				->AddTab( PreviewTabId, ETabState::OpenedTab )
+			)
+		)
+	);
+
+	const bool bCreateDefaultStandaloneMenu = true;
+	const bool bCreateDefaultToolbar = true;
+	FAssetEditorToolkit::InitAssetEditor(Mode, InitToolkitHost, FontEditorAppIdentifier, StandaloneDefaultLayout, bCreateDefaultStandaloneMenu, bCreateDefaultToolbar, ObjectToEdit);
+
+	IFontEditorModule* FontEditorModule = &FModuleManager::LoadModuleChecked<IFontEditorModule>("FontEditor");
+	AddMenuExtender(FontEditorModule->GetMenuExtensibilityManager()->GetAllExtenders(GetToolkitCommands(), GetEditingObjects()));
+}
+
+UFontFace* FFontFaceEditor::GetFontFace() const
+{
+	return FontFace;
+}
+
+FName FFontFaceEditor::GetToolkitFName() const
+{
+	return FName("FontFaceEditor");
+}
+
+FText FFontFaceEditor::GetBaseToolkitName() const
+{
+	return LOCTEXT( "AppLabel", "Font Face Editor" );
+}
+
+FString FFontFaceEditor::GetWorldCentricTabPrefix() const
+{
+	return LOCTEXT("WorldCentricTabPrefix", "Font Face ").ToString();
+}
+
+FLinearColor FFontFaceEditor::GetWorldCentricTabColorScale() const
+{
+	return FLinearColor(0.3f, 0.2f, 0.5f, 0.5f);
+}
+
+TSharedRef<SDockTab> FFontFaceEditor::SpawnTab_Preview( const FSpawnTabArgs& Args )
+{
+	check( Args.GetTabId().TabType == PreviewTabId );
+
+	TSharedRef<SDockTab> SpawnedTab = SNew(SDockTab)
+		.Label(LOCTEXT("FontFacePreviewTitle", "Preview"))
+		[
+			FontFacePreview.ToSharedRef()
+		];
+
+	AddToSpawnedToolPanels( Args.GetTabId().TabType, SpawnedTab );
+
+	return SpawnedTab;
+}
+
+TSharedRef<SDockTab> FFontFaceEditor::SpawnTab_Properties( const FSpawnTabArgs& Args )
+{
+	check( Args.GetTabId().TabType == PropertiesTabId );
+
+	TSharedRef<SDockTab> SpawnedTab = SNew(SDockTab)
+		.Label(LOCTEXT("FontFacePropertiesTitle", "Details"))
+		[
+			FontFaceProperties.ToSharedRef()
+		];
+
+	AddToSpawnedToolPanels( Args.GetTabId().TabType, SpawnedTab );
+
+	return SpawnedTab;
+}
+
+void FFontFaceEditor::AddToSpawnedToolPanels( const FName& TabIdentifier, const TSharedRef<SDockTab>& SpawnedTab )
+{
+	TWeakPtr<SDockTab>* TabSpot = SpawnedToolPanels.Find(TabIdentifier);
+	if (!TabSpot)
+	{
+		SpawnedToolPanels.Add(TabIdentifier, SpawnedTab);
+	}
+	else
+	{
+		check(!TabSpot->IsValid());
+		*TabSpot = SpawnedTab;
+	}
+}
+
+void FFontFaceEditor::AddReferencedObjects(FReferenceCollector& Collector)
+{
+	Collector.AddReferencedObject(FontFace);
+	Collector.AddReferencedObjects(PreviewFonts);
+	Collector.AddReferencedObjects(PreviewFaces);
+}
+
+void FFontFaceEditor::OnPreviewTextChanged(const FText& Text)
+{
+	for (TSharedPtr<STextBlock> &PreviewTextBlock : PreviewTextBlocks[1])
+	{
+		PreviewTextBlock->SetText(Text);
+	}
+}
+
+TOptional<int32> FFontFaceEditor::GetPreviewFontSize() const
+{
+	return PreviewFontSize;
+}
+
+void FFontFaceEditor::OnPreviewFontSizeChanged(int32 InNewValue, ETextCommit::Type CommitType)
+{
+	PreviewFontSize = InNewValue;
+	ApplyPreviewFontSize();
+}
+
+void FFontFaceEditor::NotifyPostChange( const FPropertyChangedEvent& PropertyChangedEvent, class FEditPropertyChain* PropertyThatChanged)
+{
+	static const FName EnableDistanceFieldRenderingPropertyName = GET_MEMBER_NAME_CHECKED(UFontFace, bEnableDistanceFieldRendering);
+
+	if (PropertyChangedEvent.Property && PropertyChangedEvent.Property->GetFName() == EnableDistanceFieldRenderingPropertyName)
+	{
+		// Show / hide distance field related properties
+		FontFaceProperties->ForceRefresh();
+	}
+
+	RefreshPreview();
+}
+
+void FFontFaceEditor::CreateInternalWidgets()
+{
+	const EVerticalAlignment PreviewVAlign = VAlign_Center;
+	const FText DefaultPreviewText = LOCTEXT("DefaultPreviewText", "The quick brown fox jumps over the lazy dog");
+
+	FontFacePreview =
+	SNew(SVerticalBox)
+	+SVerticalBox::Slot()
+	.FillHeight(1.0f)
+	.Padding(0.0f, 0.0f, 0.0f, 4.0f)
+	[
+		SNew(SScrollBox)
+		+SScrollBox::Slot()
+		[
+			SNew(SVerticalBox)
+			+SVerticalBox::Slot()
+			.AutoHeight()
+			[
+				SNew(SGridPanel)
+
+				+SGridPanel::Slot(0, 0)
+				.VAlign(PreviewVAlign)
+				[
+					SAssignNew(PreviewTextBlocks[0][0], STextBlock)
+					.Text(LOCTEXT("FontFaceReference", "Reference: "))
+				]
+				+SGridPanel::Slot(1, 0)
+				.VAlign(PreviewVAlign)
+				[
+					SAssignNew(PreviewTextBlocks[1][0], STextBlock)
+					.Text(DefaultPreviewText)
+				]
+				+SGridPanel::Slot(0, 1)
+				.VAlign(PreviewVAlign)
+				[
+					SAssignNew(PreviewTextBlocks[0][1], STextBlock)
+					.Text(LOCTEXT("FontFaceLowQuality", "Low: "))
+				]
+				+SGridPanel::Slot(1, 1)
+				.VAlign(PreviewVAlign)
+				[
+					SAssignNew(PreviewTextBlocks[1][1], STextBlock)
+					.Text(DefaultPreviewText)
+				]
+				+SGridPanel::Slot(0, 2)
+				.VAlign(PreviewVAlign)
+				[
+					SAssignNew(PreviewTextBlocks[0][2], STextBlock)
+					.Text(LOCTEXT("FontFaceMediumQuality", "Medium: "))
+				]
+				+SGridPanel::Slot(1, 2)
+				.VAlign(PreviewVAlign)
+				[
+					SAssignNew(PreviewTextBlocks[1][2], STextBlock)
+					.Text(DefaultPreviewText)
+				]
+				+SGridPanel::Slot(0, 3)
+				.VAlign(PreviewVAlign)
+				[
+					SAssignNew(PreviewTextBlocks[0][3], STextBlock)
+					.Text(LOCTEXT("FontFaceHighQuality", "High: "))
+				]
+				+SGridPanel::Slot(1, 3)
+				.VAlign(PreviewVAlign)
+				[
+					SAssignNew(PreviewTextBlocks[1][3], STextBlock)
+					.Text(DefaultPreviewText)
+				]
+				+SGridPanel::Slot(0, 4)
+				.VAlign(PreviewVAlign)
+				[
+					SAssignNew(PreviewTextBlocks[0][4], STextBlock)
+					.Text(LOCTEXT("FontFaceMultiLowQuality", "Multi Low: "))
+				]
+				+SGridPanel::Slot(1, 4)
+				.VAlign(PreviewVAlign)
+				[
+					SAssignNew(PreviewTextBlocks[1][4], STextBlock)
+					.Text(DefaultPreviewText)
+				]
+				+SGridPanel::Slot(0, 5)
+				.VAlign(PreviewVAlign)
+				[
+					SAssignNew(PreviewTextBlocks[0][5], STextBlock)
+					.Text(LOCTEXT("FontFaceMultiMediumQuality", "Multi Medium: "))
+				]
+				+SGridPanel::Slot(1, 5)
+				.VAlign(PreviewVAlign)
+				[
+					SAssignNew(PreviewTextBlocks[1][5], STextBlock)
+					.Text(DefaultPreviewText)
+				]
+				+SGridPanel::Slot(0, 6)
+				.VAlign(PreviewVAlign)
+				[
+					SAssignNew(PreviewTextBlocks[0][6], STextBlock)
+					.Text(LOCTEXT("FontFaceMultiHighQuality", "Multi High: "))
+				]
+				+SGridPanel::Slot(1, 6)
+				.VAlign(PreviewVAlign)
+				[
+					SAssignNew(PreviewTextBlocks[1][6], STextBlock)
+					.Text(DefaultPreviewText)
+				]
+			]
+			+SVerticalBox::Slot()
+			.AutoHeight()
+			[
+				SAssignNew(PreviewNoteTextBlock, STextBlock)
+				.Text(LOCTEXT("FontFaceDistanceFieldProjectSettingNote", "Note: You must also enable Distance Field Font Rasterization in Project Settings / Engine / User Interface."))
+				.Visibility(EVisibility::Collapsed)
+			]
+		]
+	]
+	+SVerticalBox::Slot()
+	.AutoHeight()
+	[
+		SNew(SHorizontalBox)
+
+		+SHorizontalBox::Slot()
+		[
+			SAssignNew(FontFacePreviewText, SEditableTextBox)
+			.Text(DefaultPreviewText)
+			.SelectAllTextWhenFocused(true)
+			.OnTextChanged(this, &FFontFaceEditor::OnPreviewTextChanged)
+		]
+
+		+ SHorizontalBox::Slot()
+		.AutoWidth()
+		[
+			SNew(SNumericEntryBox<int32>)
+			.Value(this, &FFontFaceEditor::GetPreviewFontSize)
+			.MinValue(4)
+			.MaxValue(256)
+			.OnValueCommitted(this, &FFontFaceEditor::OnPreviewFontSizeChanged)
+		]
+	];
+
+	UpdatePreviewFonts();
+	UpdatePreviewVisibility();
+	ApplyPreviewFontSize();
+
+	FDetailsViewArgs Args;
+	Args.bHideSelectionTip = true;
+	Args.NotifyHook = this;
+
+	FPropertyEditorModule& PropertyModule = FModuleManager::LoadModuleChecked<FPropertyEditorModule>("PropertyEditor");
+	FontFaceProperties = PropertyModule.CreateDetailView(Args);
+
+	FontFaceProperties->SetIsPropertyVisibleDelegate(FIsPropertyVisible::CreateRaw(this, &FFontFaceEditor::GetIsPropertyVisible));
+	FontFaceProperties->SetObject( FontFace );
+}
+
+void FFontFaceEditor::OnPostReimport(UObject* InObject, bool bSuccess)
+{
+	// Ignore if this is regarding a different object
+	if ( InObject != FontFace )
+	{
+		return;
+	}
+
+	if ( bSuccess )
+	{
+		//FontFaceViewport->RefreshViewport();
+		RefreshPreview();
+	}
+}
+
+void FFontFaceEditor::OnObjectPropertyChanged(UObject* InObject, struct FPropertyChangedEvent& InPropertyChangedEvent)
+{
+	if (Cast<UFontFace>(InObject))
+	{
+		// Force all texts using a font to be refreshed.
+		FSlateApplicationBase::Get().InvalidateAllWidgets(false);
+		GSlateLayoutGeneration++;
+	}
+}
+
+void FFontFaceEditor::OnObjectReimported(UObject* InObject)
+{
+	// Make sure we are using the object that is being reimported, otherwise a lot of needless work could occur.
+	if(FontFace == InObject)
+	{
+		FontFace = Cast<UFontFace>(InObject);
+
+		TArray< UObject* > ObjectList;
+		ObjectList.Add(InObject);
+		FontFaceProperties->SetObjects(ObjectList);
+	}
+}
+
+bool FFontFaceEditor::GetIsPropertyVisible(const FPropertyAndParent& PropertyAndParent) const
+{
+	static const FName CategoryFName = "Category";
+	const FString& CategoryValue = PropertyAndParent.Property.GetMetaData(CategoryFName);
+	return CategoryValue != TEXT("DistanceFieldMode") || IsSlateSdfTextFeatureEnabled();
+}
+
+bool FFontFaceEditor::ShouldPromptForNewFilesOnReload(const UObject& EditingObject) const
+{
+	return false;
+}
+
+void FFontFaceEditor::RefreshPreview()
+{
+	UpdatePreviewFonts();
+	UpdatePreviewVisibility();
+}
+
+void FFontFaceEditor::ClonePreviewFontFace(TObjectPtr<UFontFace>& TargetFontFace, EFontRasterizationMode RasterizationMode, int32 DistanceFieldPpem) const
+{
+	TargetFontFace = DuplicateObject<UFontFace>(FontFace, GetTransientPackage());
+	TargetFontFace->MinDistanceFieldPpem = DistanceFieldPpem;
+	TargetFontFace->MidDistanceFieldPpem = DistanceFieldPpem;
+	TargetFontFace->MaxDistanceFieldPpem = DistanceFieldPpem;
+	TargetFontFace->MinMultiDistanceFieldPpem = DistanceFieldPpem;
+	TargetFontFace->MidMultiDistanceFieldPpem = DistanceFieldPpem;
+	TargetFontFace->MaxMultiDistanceFieldPpem = DistanceFieldPpem;
+	TargetFontFace->PlatformRasterizationModeOverrides = FFontFacePlatformRasterizationOverrides();
+	TargetFontFace->PlatformRasterizationModeOverrides->MsdfOverride = RasterizationMode;
+	TargetFontFace->PlatformRasterizationModeOverrides->SdfOverride = RasterizationMode;
+	TargetFontFace->PlatformRasterizationModeOverrides->SdfApproximationOverride = RasterizationMode;
+	TargetFontFace->PostEditChange();
+}
+
+void FFontFaceEditor::MakePreviewFont(TObjectPtr<UObject>& TargetObject, UFontFace* Face) const
+{
+	if (!TargetObject)
+	{
+		TargetObject = NewObject<UFont>();
+	}
+	if (UFont* TargetFont = CastChecked<UFont>(TargetObject))
+	{
+		if (TargetFont->CompositeFont.DefaultTypeface.Fonts.IsEmpty())
+		{
+			FTypefaceEntry FontTypeface;
+			FontTypeface.Name = TEXT("Regular");
+			FontTypeface.Font = FFontData(Face);
+			TargetFont->FontCacheType = EFontCacheType::Runtime;
+			TargetFont->CompositeFont.DefaultTypeface.Fonts.Add(MoveTemp(FontTypeface));
+		}
+		else
+		{
+			TargetFont->CompositeFont.DefaultTypeface.Fonts[0].Font = FFontData(Face);
+		}
+		TargetFont->PostEditChange();
+	}
+}
+
+bool FFontFaceEditor::IsFontFaceDistanceFieldEnabled() const
+{
+	return FontFace->bEnableDistanceFieldRendering &&
+	       GetDefault<UUserInterfaceSettings>()->bEnableDistanceFieldFontRasterization &&
+	       IsSlateSdfTextFeatureEnabled();
+}
+
+void FFontFaceEditor::UpdatePreviewFonts()
+{
+	if (!FontFace)
+	{
+		return;
+	}
+	const int PreviewFontNum = UE_ARRAY_COUNT(*PreviewTextBlocks);
+	if (IsFontFaceDistanceFieldEnabled())
+	{
+		PreviewFaces.SetNum(PreviewFontNum, EAllowShrinking::No);
+		PreviewFonts.SetNum(PreviewFontNum, EAllowShrinking::No);
+		ClonePreviewFontFace(PreviewFaces[0], EFontRasterizationMode::Bitmap);
+		ClonePreviewFontFace(PreviewFaces[1], EFontRasterizationMode::Sdf, FontFace->MinDistanceFieldPpem);
+		ClonePreviewFontFace(PreviewFaces[2], EFontRasterizationMode::Sdf, FontFace->MidDistanceFieldPpem);
+		ClonePreviewFontFace(PreviewFaces[3], EFontRasterizationMode::Sdf, FontFace->MaxDistanceFieldPpem);
+		ClonePreviewFontFace(PreviewFaces[4], EFontRasterizationMode::Msdf, FontFace->MinMultiDistanceFieldPpem);
+		ClonePreviewFontFace(PreviewFaces[5], EFontRasterizationMode::Msdf, FontFace->MidMultiDistanceFieldPpem);
+		ClonePreviewFontFace(PreviewFaces[6], EFontRasterizationMode::Msdf, FontFace->MaxMultiDistanceFieldPpem);
+		for (int32 Index = 0; Index < PreviewFontNum; ++Index)
+		{
+			MakePreviewFont(PreviewFonts[Index], PreviewFaces[Index]);
+		}
+	}
+	else
+	{
+		PreviewFaces.SetNum(1, EAllowShrinking::No);
+		PreviewFonts.SetNum(PreviewFontNum, EAllowShrinking::No);
+		ClonePreviewFontFace(PreviewFaces[0], EFontRasterizationMode::Bitmap);
+		for (TObjectPtr<UObject>& PreviewFont : PreviewFonts)
+		{
+			MakePreviewFont(PreviewFont, PreviewFaces[0]);
+		}
+	}
+}
+
+void FFontFaceEditor::UpdatePreviewVisibility()
+{
+	if (FontFace)
+	{
+		const EVisibility SecondaryRowsVisibility = IsFontFaceDistanceFieldEnabled() ? EVisibility::Visible : EVisibility::Collapsed;
+		PreviewTextBlocks[0][0]->SetVisibility(SecondaryRowsVisibility);
+		for (int32 Index = 1; Index < UE_ARRAY_COUNT(*PreviewTextBlocks); ++Index)
+		{
+			PreviewTextBlocks[0][Index]->SetVisibility(SecondaryRowsVisibility);
+			PreviewTextBlocks[1][Index]->SetVisibility(SecondaryRowsVisibility);
+		}
+		PreviewNoteTextBlock->SetVisibility(
+			FontFace->bEnableDistanceFieldRendering &&
+			IsSlateSdfTextFeatureEnabled() &&
+			!GetDefault<UUserInterfaceSettings>()->bEnableDistanceFieldFontRasterization ?
+			EVisibility::Visible : EVisibility::Collapsed
+		);
+	}
+}
+
+void FFontFaceEditor::ApplyPreviewFontSize()
+{
+	for (int32 Index = 0; Index < UE_ARRAY_COUNT(PreviewTextBlocks[1]) && Index < PreviewFonts.Num(); ++Index)
+	{
+		PreviewTextBlocks[1][Index]->SetFont(FSlateFontInfo(PreviewFonts[Index], PreviewFontSize));
+	}
+}
+
+#undef LOCTEXT_NAMESPACE
diff --git a/Engine/Source/Editor/FontEditor/Private/FontFaceEditor.h b/Engine/Source/Editor/FontEditor/Private/FontFaceEditor.h
new file mode 100644
--- /dev/null
+++ b/Engine/Source/Editor/FontEditor/Private/FontFaceEditor.h
@@ -0,0 +1,142 @@
+// Copyright Epic Games, Inc. All Rights Reserved.
+
+#pragma once
+
+#include "EditorUndoClient.h"
+#include "Framework/Docking/TabManager.h"
+#include "HAL/Platform.h"
+#include "IFontFaceEditor.h"
+#include "Internationalization/Text.h"
+#include "Math/Color.h"
+#include "Misc/NotifyHook.h"
+#include "Misc/Optional.h"
+#include "Styling/SlateTypes.h"
+#include "Templates/SharedPointer.h"
+#include "Toolkits/IToolkit.h"
+#include "UObject/GCObject.h"
+#include "UObject/NameTypes.h"
+#include "UObject/Object.h"
+
+class FReferenceCollector;
+class SCompositeFontEditor;
+class SDockTab;
+class SEditableTextBox;
+class SVerticalBox;
+class UFontFace;
+struct FPropertyChangedEvent;
+
+/*-----------------------------------------------------------------------------
+   FFontFaceEditor
+-----------------------------------------------------------------------------*/
+
+class FFontFaceEditor : public IFontFaceEditor, public FGCObject, public FNotifyHook, public FEditorUndoClient
+{
+public:
+	FFontFaceEditor();
+
+	virtual void RegisterTabSpawners(const TSharedRef<class FTabManager>& TabManager) override;
+	virtual void UnregisterTabSpawners(const TSharedRef<class FTabManager>& TabManager) override;
+
+	/** Destructor */
+	virtual ~FFontFaceEditor();
+
+	/** Edits the specified Font object */
+	void InitFontFaceEditor(const EToolkitMode::Type Mode, const TSharedPtr< class IToolkitHost >& InitToolkitHost, UObject* ObjectToEdit);
+
+	/** IFontFaceEditor interface */
+	virtual UFontFace* GetFontFace() const override;
+	virtual void RefreshPreview() override;
+	
+	/** IToolkit interface */
+	virtual FName GetToolkitFName() const override;
+	virtual FText GetBaseToolkitName() const override;
+	virtual FString GetWorldCentricTabPrefix() const override;
+	virtual FLinearColor GetWorldCentricTabColorScale() const override;
+
+	/** FGCObject interface */
+	virtual void AddReferencedObjects(FReferenceCollector& Collector) override;
+	virtual FString GetReferencerName() const override
+	{
+		return TEXT("FFontFaceEditor");
+	}
+
+	/** Called to determine if the user should be prompted for a new file if one is missing during an asset reload*/
+	virtual bool ShouldPromptForNewFilesOnReload(const UObject& object) const override;
+
+protected:
+	/** Called when the preview text changes */
+	void OnPreviewTextChanged(const FText& Text);
+
+	/** Called to handle the "Preview Font Size" numeric entry box */
+	TOptional<int32> GetPreviewFontSize() const;
+	void OnPreviewFontSizeChanged(int32 InNewValue, ETextCommit::Type CommitType);
+
+private:
+	/** FNotifyHook interface */
+	virtual void NotifyPostChange(const FPropertyChangedEvent& PropertyChangedEvent, class FEditPropertyChain* PropertyThatChanged) override;
+
+	/** Creates all internal widgets for the tabs to point at */
+	void CreateInternalWidgets();
+
+	void OnPostReimport(UObject* InObject, bool bSuccess);
+	void OnObjectPropertyChanged(UObject* InObject, struct FPropertyChangedEvent& InPropertyChangedEvent);
+
+	/** Spawns the preview tab */
+	TSharedRef<SDockTab> SpawnTab_Preview(const FSpawnTabArgs& Args);
+
+	/** Spawns the properties tab */
+	TSharedRef<SDockTab> SpawnTab_Properties(const FSpawnTabArgs& Args);
+
+	/** Caches the specified tab for later retrieval */
+	void AddToSpawnedToolPanels(const FName& TabIdentifier, const TSharedRef<SDockTab>& SpawnedTab);
+
+	/** Callback when an object is reimported, handles steps needed to keep the editor up-to-date. */
+	void OnObjectReimported(UObject* InObject);
+
+	/** Check to see if the given property should be visible in the details panel */
+	bool GetIsPropertyVisible(const struct FPropertyAndParent& PropertyAndParent) const;
+
+private:
+	/** The font asset being inspected */
+	TObjectPtr<UFontFace> FontFace;
+
+	/** Virtual fonts for editor preview only */
+	TArray<TObjectPtr<UObject>> PreviewFonts;
+
+	/** Virtual font faces for editor preview only */
+	TArray<TObjectPtr<UFontFace>> PreviewFaces;
+
+	/** List of open tool panels; used to ensure only one exists at any one time */
+	TMap< FName, TWeakPtr<SDockTab> > SpawnedToolPanels;
+
+	/** Preview tab */
+	TSharedPtr<SVerticalBox> FontFacePreview;
+
+	/** Properties tab */
+	TSharedPtr<class IDetailsView> FontFaceProperties;
+
+	/** Preview text */
+	TSharedPtr<SEditableTextBox> FontFacePreviewText;
+
+	/** Preview widgets */
+	TSharedPtr<STextBlock> PreviewTextBlocks[2][7];
+
+	/** Preview note text */
+	TSharedPtr<STextBlock> PreviewNoteTextBlock;
+
+	/** Preview font size */
+	int32 PreviewFontSize = 24;
+
+	/** The tab ids for the font editor */
+	static const FName PreviewTabId;
+	static const FName PropertiesTabId;
+
+	void ClonePreviewFontFace(TObjectPtr<UFontFace>& TargetFontFace, EFontRasterizationMode RasterizationMode, int32 DistanceFieldPpem = 0) const;
+	void MakePreviewFont(TObjectPtr<UObject>& TargetObject, UFontFace* Face) const;
+	bool IsFontFaceDistanceFieldEnabled() const;
+
+	void UpdatePreviewFonts();
+	void UpdatePreviewVisibility();
+	void ApplyPreviewFontSize();
+
+};
diff --git a/Engine/Source/Editor/FontEditor/Public/FontEditorModule.h b/Engine/Source/Editor/FontEditor/Public/FontEditorModule.h
--- a/Engine/Source/Editor/FontEditor/Public/FontEditorModule.h
+++ b/Engine/Source/Editor/FontEditor/Public/FontEditorModule.h
@@ -10,8 +10,10 @@
 #include "UObject/NameTypes.h"
 
 class IFontEditor;
+class IFontFaceEditor;
 class IToolkitHost;
 class UFont;
+class UFontFace;
 
 extern const FName FontEditorAppIdentifier;
 
@@ -26,4 +28,7 @@ class IFontEditorModule : public IModuleInterface,
 public:
 	/** Creates a new Font editor */
 	virtual TSharedRef<IFontEditor> CreateFontEditor( const EToolkitMode::Type Mode, const TSharedPtr< IToolkitHost >& InitToolkitHost, UFont* Font ) = 0;
+
+	/** Creates a new Font Face editor */
+	virtual TSharedRef<IFontFaceEditor> CreateFontFaceEditor( const EToolkitMode::Type Mode, const TSharedPtr< IToolkitHost >& InitToolkitHost, UFontFace* FontFace ) = 0;
 };
diff --git a/Engine/Source/Editor/FontEditor/Public/IFontFaceEditor.h b/Engine/Source/Editor/FontEditor/Public/IFontFaceEditor.h
new file mode 100644
--- /dev/null
+++ b/Engine/Source/Editor/FontEditor/Public/IFontFaceEditor.h
@@ -0,0 +1,22 @@
+// Copyright Epic Games, Inc. All Rights Reserved.
+
+#pragma once
+
+#include "CoreMinimal.h"
+#include "Toolkits/AssetEditorToolkit.h"
+
+class UFontFace;
+
+/*-----------------------------------------------------------------------------
+   IFontFaceEditor
+-----------------------------------------------------------------------------*/
+
+class IFontFaceEditor : public FAssetEditorToolkit
+{
+public:
+	/** Returns the font face asset being inspected by the font face editor */
+	virtual UFontFace* GetFontFace() const = 0;
+
+	/** Refresh the preview viewport */
+	virtual void RefreshPreview() = 0;
+};
