Font Face editor fix - update preview on source file change
#tests Tested in editor
#rnx <Release Notes exclude>

diff --git a/Engine/Source/Editor/FontEditor/Private/FontFaceDetailsCustomization.cpp b/Engine/Source/Editor/FontEditor/Private/FontFaceDetailsCustomization.cpp
--- a/Engine/Source/Editor/FontEditor/Private/FontFaceDetailsCustomization.cpp
+++ b/Engine/Source/Editor/FontEditor/Private/FontFaceDetailsCustomization.cpp
@@ -208,6 +208,7 @@ void FFontFaceDetailsCustomization::OnFontPathPicked(const FString& InNewFontFil
 			FontFace->SourceFilename = InNewFontFilename;
 			FontFace->FontFaceData = FFontFaceData::MakeFontFaceData(MoveTemp(FontData)); // Make a new instance as the existing one may be being used by the font cache
 			FontFace->CacheSubFaces();
+			FontFace->PostEditChange();
 		}
 	}
 
diff --git a/Engine/Source/Editor/FontEditor/Private/FontFaceEditor.cpp b/Engine/Source/Editor/FontEditor/Private/FontFaceEditor.cpp
--- a/Engine/Source/Editor/FontEditor/Private/FontFaceEditor.cpp
+++ b/Engine/Source/Editor/FontEditor/Private/FontFaceEditor.cpp
@@ -221,7 +221,7 @@ void FFontFaceEditor::OnPreviewFontSizeChanged(int32 InNewValue, ETextCommit::Ty
 	ApplyPreviewFontSize();
 }
 
-void FFontFaceEditor::NotifyPostChange( const FPropertyChangedEvent& PropertyChangedEvent, class FEditPropertyChain* PropertyThatChanged)
+void FFontFaceEditor::NotifyPostChange(const FPropertyChangedEvent& PropertyChangedEvent, class FEditPropertyChain* PropertyThatChanged)
 {
 	static const FName EnableDistanceFieldRenderingPropertyName = GET_MEMBER_NAME_CHECKED(UFontFace, bEnableDistanceFieldRendering);
 
@@ -389,33 +389,27 @@ void FFontFaceEditor::CreateInternalWidgets()
 
 void FFontFaceEditor::OnPostReimport(UObject* InObject, bool bSuccess)
 {
-	// Ignore if this is regarding a different object
-	if ( InObject != FontFace )
+	if (InObject == FontFace && bSuccess)
 	{
-		return;
-	}
-
-	if ( bSuccess )
-	{
-		//FontFaceViewport->RefreshViewport();
 		RefreshPreview();
 	}
 }
 
 void FFontFaceEditor::OnObjectPropertyChanged(UObject* InObject, struct FPropertyChangedEvent& InPropertyChangedEvent)
 {
-	if (Cast<UFontFace>(InObject))
+	if (InObject == FontFace)
 	{
 		// Force all texts using a font to be refreshed.
 		FSlateApplicationBase::Get().InvalidateAllWidgets(false);
 		GSlateLayoutGeneration++;
+		RefreshPreview();
 	}
 }
 
 void FFontFaceEditor::OnObjectReimported(UObject* InObject)
 {
 	// Make sure we are using the object that is being reimported, otherwise a lot of needless work could occur.
-	if(FontFace == InObject)
+	if (InObject == FontFace)
 	{
 		FontFace = Cast<UFontFace>(InObject);
 
