Font Face Editor - initial layout fix
#tests Tested in editor
#rnx <Release Notes exclude>

diff --git a/Engine/Source/Editor/FontEditor/Private/FontFaceEditor.cpp b/Engine/Source/Editor/FontEditor/Private/FontFaceEditor.cpp
--- a/Engine/Source/Editor/FontEditor/Private/FontFaceEditor.cpp
+++ b/Engine/Source/Editor/FontEditor/Private/FontFaceEditor.cpp
@@ -98,7 +98,7 @@ void FFontFaceEditor::InitFontFaceEditor(const EToolkitMode::Type Mode, const TS
 
 	CreateInternalWidgets();
 
-	const TSharedRef<FTabManager::FLayout> StandaloneDefaultLayout = FTabManager::NewLayout("Standalone_FontEditor_Layout_v4")
+	const TSharedRef<FTabManager::FLayout> StandaloneDefaultLayout = FTabManager::NewLayout("Standalone_FontFaceEditor_Layout_v1")
 	->AddArea
 	(
 		FTabManager::NewPrimaryArea() ->SetOrientation( Orient_Vertical )
