[SDF Fonts] Font Face Editor improvements:
- Added horizontal scrollbox
- Added preview text padding

#tests Tested in editor
#rnx

diff --git a/Engine/Source/Editor/FontEditor/Private/FontFaceEditor.cpp b/Engine/Source/Editor/FontEditor/Private/FontFaceEditor.cpp
--- a/Engine/Source/Editor/FontEditor/Private/FontFaceEditor.cpp
+++ b/Engine/Source/Editor/FontEditor/Private/FontFaceEditor.cpp
@@ -252,91 +252,95 @@ void FFontFaceEditor::CreateInternalWidgets()
 			+SVerticalBox::Slot()
 			.AutoHeight()
 			[
-				SNew(SGridPanel)
-
-				+SGridPanel::Slot(0, 0)
-				.VAlign(PreviewVAlign)
-				[
-					SAssignNew(PreviewTextBlocks[0][0], STextBlock)
-					.Text(LOCTEXT("FontFaceReference", "Reference: "))
-				]
-				+SGridPanel::Slot(1, 0)
-				.VAlign(PreviewVAlign)
-				[
-					SAssignNew(PreviewTextBlocks[1][0], STextBlock)
-					.Text(DefaultPreviewText)
-				]
-				+SGridPanel::Slot(0, 1)
-				.VAlign(PreviewVAlign)
-				[
-					SAssignNew(PreviewTextBlocks[0][1], STextBlock)
-					.Text(LOCTEXT("FontFaceLowQuality", "Low: "))
-				]
-				+SGridPanel::Slot(1, 1)
-				.VAlign(PreviewVAlign)
-				[
-					SAssignNew(PreviewTextBlocks[1][1], STextBlock)
-					.Text(DefaultPreviewText)
-				]
-				+SGridPanel::Slot(0, 2)
-				.VAlign(PreviewVAlign)
-				[
-					SAssignNew(PreviewTextBlocks[0][2], STextBlock)
-					.Text(LOCTEXT("FontFaceMediumQuality", "Medium: "))
-				]
-				+SGridPanel::Slot(1, 2)
-				.VAlign(PreviewVAlign)
-				[
-					SAssignNew(PreviewTextBlocks[1][2], STextBlock)
-					.Text(DefaultPreviewText)
-				]
-				+SGridPanel::Slot(0, 3)
-				.VAlign(PreviewVAlign)
-				[
-					SAssignNew(PreviewTextBlocks[0][3], STextBlock)
-					.Text(LOCTEXT("FontFaceHighQuality", "High: "))
-				]
-				+SGridPanel::Slot(1, 3)
-				.VAlign(PreviewVAlign)
-				[
-					SAssignNew(PreviewTextBlocks[1][3], STextBlock)
-					.Text(DefaultPreviewText)
-				]
-				+SGridPanel::Slot(0, 4)
-				.VAlign(PreviewVAlign)
+				SNew(SScrollBox)
+				.Orientation(EOrientation::Orient_Horizontal)
+				+SScrollBox::Slot()
 				[
-					SAssignNew(PreviewTextBlocks[0][4], STextBlock)
-					.Text(LOCTEXT("FontFaceMultiLowQuality", "Multi Low: "))
-				]
-				+SGridPanel::Slot(1, 4)
-				.VAlign(PreviewVAlign)
-				[
-					SAssignNew(PreviewTextBlocks[1][4], STextBlock)
-					.Text(DefaultPreviewText)
-				]
-				+SGridPanel::Slot(0, 5)
-				.VAlign(PreviewVAlign)
-				[
-					SAssignNew(PreviewTextBlocks[0][5], STextBlock)
-					.Text(LOCTEXT("FontFaceMultiMediumQuality", "Multi Medium: "))
-				]
-				+SGridPanel::Slot(1, 5)
-				.VAlign(PreviewVAlign)
-				[
-					SAssignNew(PreviewTextBlocks[1][5], STextBlock)
-					.Text(DefaultPreviewText)
-				]
-				+SGridPanel::Slot(0, 6)
-				.VAlign(PreviewVAlign)
-				[
-					SAssignNew(PreviewTextBlocks[0][6], STextBlock)
-					.Text(LOCTEXT("FontFaceMultiHighQuality", "Multi High: "))
-				]
-				+SGridPanel::Slot(1, 6)
-				.VAlign(PreviewVAlign)
-				[
-					SAssignNew(PreviewTextBlocks[1][6], STextBlock)
-					.Text(DefaultPreviewText)
+					SAssignNew(PreviewTextGridPanel, SGridPanel)
+					+SGridPanel::Slot(0, 0)
+					.VAlign(PreviewVAlign)
+					[
+						SAssignNew(PreviewTextBlocks[0][0], STextBlock)
+						.Text(LOCTEXT("FontFaceReference", "Reference: "))
+					]
+					+SGridPanel::Slot(1, 0)
+					.VAlign(PreviewVAlign)
+					[
+						SAssignNew(PreviewTextBlocks[1][0], STextBlock)
+						.Text(DefaultPreviewText)
+					]
+					+SGridPanel::Slot(0, 1)
+					.VAlign(PreviewVAlign)
+					[
+						SAssignNew(PreviewTextBlocks[0][1], STextBlock)
+						.Text(LOCTEXT("FontFaceLowQuality", "Low: "))
+					]
+					+SGridPanel::Slot(1, 1)
+					.VAlign(PreviewVAlign)
+					[
+						SAssignNew(PreviewTextBlocks[1][1], STextBlock)
+						.Text(DefaultPreviewText)
+					]
+					+SGridPanel::Slot(0, 2)
+					.VAlign(PreviewVAlign)
+					[
+						SAssignNew(PreviewTextBlocks[0][2], STextBlock)
+						.Text(LOCTEXT("FontFaceMediumQuality", "Medium: "))
+					]
+					+SGridPanel::Slot(1, 2)
+					.VAlign(PreviewVAlign)
+					[
+						SAssignNew(PreviewTextBlocks[1][2], STextBlock)
+						.Text(DefaultPreviewText)
+					]
+					+SGridPanel::Slot(0, 3)
+					.VAlign(PreviewVAlign)
+					[
+						SAssignNew(PreviewTextBlocks[0][3], STextBlock)
+						.Text(LOCTEXT("FontFaceHighQuality", "High: "))
+					]
+					+SGridPanel::Slot(1, 3)
+					.VAlign(PreviewVAlign)
+					[
+						SAssignNew(PreviewTextBlocks[1][3], STextBlock)
+						.Text(DefaultPreviewText)
+					]
+					+SGridPanel::Slot(0, 4)
+					.VAlign(PreviewVAlign)
+					[
+						SAssignNew(PreviewTextBlocks[0][4], STextBlock)
+						.Text(LOCTEXT("FontFaceMultiLowQuality", "Multi Low: "))
+					]
+					+SGridPanel::Slot(1, 4)
+					.VAlign(PreviewVAlign)
+					[
+						SAssignNew(PreviewTextBlocks[1][4], STextBlock)
+						.Text(DefaultPreviewText)
+					]
+					+SGridPanel::Slot(0, 5)
+					.VAlign(PreviewVAlign)
+					[
+						SAssignNew(PreviewTextBlocks[0][5], STextBlock)
+						.Text(LOCTEXT("FontFaceMultiMediumQuality", "Multi Medium: "))
+					]
+					+SGridPanel::Slot(1, 5)
+					.VAlign(PreviewVAlign)
+					[
+						SAssignNew(PreviewTextBlocks[1][5], STextBlock)
+						.Text(DefaultPreviewText)
+					]
+					+SGridPanel::Slot(0, 6)
+					.VAlign(PreviewVAlign)
+					[
+						SAssignNew(PreviewTextBlocks[0][6], STextBlock)
+						.Text(LOCTEXT("FontFaceMultiHighQuality", "Multi High: "))
+					]
+					+SGridPanel::Slot(1, 6)
+					.VAlign(PreviewVAlign)
+					[
+						SAssignNew(PreviewTextBlocks[1][6], STextBlock)
+						.Text(DefaultPreviewText)
+					]
 				]
 			]
 			+SVerticalBox::Slot()
@@ -543,10 +547,34 @@ void FFontFaceEditor::UpdatePreviewVisibility()
 
 void FFontFaceEditor::ApplyPreviewFontSize()
 {
-	for (int32 Index = 0; Index < UE_ARRAY_COUNT(PreviewTextBlocks[1]) && Index < PreviewFonts.Num(); ++Index)
+	constexpr const int32 ColumnIndex = 1;
+	for (int32 RowIndex = 0; RowIndex < UE_ARRAY_COUNT(PreviewTextBlocks[ColumnIndex]) && RowIndex < PreviewFonts.Num(); ++RowIndex)
 	{
-		PreviewTextBlocks[1][Index]->SetFont(FSlateFontInfo(PreviewFonts[Index], PreviewFontSize));
+		const TSharedPtr<STextBlock>& PreviewTextBlock = PreviewTextBlocks[ColumnIndex][RowIndex];
+		PreviewTextBlock->SetFont(FSlateFontInfo(PreviewFonts[RowIndex], PreviewFontSize));
+
+		if (TPanelChildren<SGridPanel::FSlot>* PreviewTextGridPanelChildren = static_cast<TPanelChildren<SGridPanel::FSlot>*>(PreviewTextGridPanel->GetChildren()))
+		{
+			TPanelChildren<SGridPanel::FSlot>& Children = *PreviewTextGridPanelChildren;
+			TSharedRef<STextBlock> TextBlock = PreviewTextBlock.ToSharedRef();
+			for (int32 SlotIndex = 0; SlotIndex < Children.Num(); ++SlotIndex)
+			{
+				SGridPanel::FSlot& Slot = Children[SlotIndex];
+				if (Slot.GetWidget() == TextBlock)
+				{
+					Slot.SetPadding(GetPreviewTextPadding());
+					break;
+				}
+			}
+		}
+
+		PreviewTextGridPanel->Invalidate(EInvalidateWidgetReason::Layout);
 	}
 }
 
+float FFontFaceEditor::GetPreviewTextPadding() const
+{
+	return PreviewFontSize / 2.0f;
+}
+
 #undef LOCTEXT_NAMESPACE
diff --git a/Engine/Source/Editor/FontEditor/Private/FontFaceEditor.h b/Engine/Source/Editor/FontEditor/Private/FontFaceEditor.h
--- a/Engine/Source/Editor/FontEditor/Private/FontFaceEditor.h
+++ b/Engine/Source/Editor/FontEditor/Private/FontFaceEditor.h
@@ -21,6 +21,7 @@ class FReferenceCollector;
 class SCompositeFontEditor;
 class SDockTab;
 class SEditableTextBox;
+class SGridPanel;
 class SVerticalBox;
 class UFontFace;
 struct FPropertyChangedEvent;
@@ -121,6 +122,9 @@ private:
 	/** Preview widgets */
 	TSharedPtr<STextBlock> PreviewTextBlocks[2][7];
 
+	/** Preview grid panel */
+	TSharedPtr<SGridPanel> PreviewTextGridPanel;
+
 	/** Preview note text */
 	TSharedPtr<STextBlock> PreviewNoteTextBlock;
 
@@ -138,5 +142,5 @@ private:
 	void UpdatePreviewFonts();
 	void UpdatePreviewVisibility();
 	void ApplyPreviewFontSize();
-
+	float GetPreviewTextPadding() const;
 };
