Font face editor improvements:
- Fixes preview scrolling so that mouse wheel scrolls vertically
- Adds visibility toggles for individual preview rows
- Adds approximate SDF preview (initially hidden by visibility toggle)
- Improves code readability by utilizing an enum for preview rows

Original author is Viktor.Chlumsky. Unshelved from CL 38973140.

#tests Tested in editor when opening font faces
#rnx

diff --git a/Engine/Source/Editor/FontEditor/Private/FontFaceEditor.cpp b/Engine/Source/Editor/FontEditor/Private/FontFaceEditor.cpp
--- a/Engine/Source/Editor/FontEditor/Private/FontFaceEditor.cpp
+++ b/Engine/Source/Editor/FontEditor/Private/FontFaceEditor.cpp
@@ -22,6 +22,7 @@
 #include "Toolkits/AssetEditorToolkit.h"
 #include "UObject/ObjectPtr.h"
 #include "Widgets/Docking/SDockTab.h"
+#include "Widgets/Input/SComboButton.h"
 #include "Widgets/Input/SEditableTextBox.h"
 #include "Widgets/Input/SNumericEntryBox.h"
 #include "Widgets/Layout/SGridPanel.h"
@@ -64,6 +65,16 @@ void FFontFaceEditor::UnregisterTabSpawners(const TSharedRef<class FTabManager>&
 FFontFaceEditor::FFontFaceEditor()
 	: FontFace(nullptr)
 {
+	PreviewRowVisibility[(int32)EPreviewRow::Reference] = true;
+	PreviewRowVisibility[(int32)EPreviewRow::ApproximateSdfLow] = false;
+	PreviewRowVisibility[(int32)EPreviewRow::ApproximateSdfMedium] = false;
+	PreviewRowVisibility[(int32)EPreviewRow::ApproximateSdfHigh] = false;
+	PreviewRowVisibility[(int32)EPreviewRow::SdfLow] = true;
+	PreviewRowVisibility[(int32)EPreviewRow::SdfMedium] = true;
+	PreviewRowVisibility[(int32)EPreviewRow::SdfHigh] = true;
+	PreviewRowVisibility[(int32)EPreviewRow::MsdfLow] = true;
+	PreviewRowVisibility[(int32)EPreviewRow::MsdfMedium] = true;
+	PreviewRowVisibility[(int32)EPreviewRow::MsdfHigh] = true;
 }
 
 FFontFaceEditor::~FFontFaceEditor()
@@ -239,6 +250,75 @@ void FFontFaceEditor::CreateInternalWidgets()
 	const EVerticalAlignment PreviewVAlign = VAlign_Center;
 	const FText DefaultPreviewText = LOCTEXT("DefaultPreviewText", "The quick brown fox jumps over the lazy dog");
 
+	FMenuBuilder PreviewRowVisibilitySelection(false, nullptr);
+
+	auto AddPreviewVisibilityItem = [this, &PreviewRowVisibilitySelection](EPreviewRow Row, const TAttribute<FText>& InLabel, const TAttribute<FText>& InToolTip)
+	{
+		PreviewRowVisibilitySelection.AddMenuEntry(
+			InLabel,
+			InToolTip,
+			FSlateIcon(),
+			FUIAction(
+				FExecuteAction::CreateSP(this, &FFontFaceEditor::ChangePreviewRowVisibility, Row),
+				FCanExecuteAction(),
+				FIsActionChecked::CreateSP(this, &FFontFaceEditor::GetPreviewRowVisibility, Row)
+			),
+			NAME_None,
+			EUserInterfaceActionType::Check
+		);
+	};
+
+	AddPreviewVisibilityItem(
+		EPreviewRow::Reference,
+		LOCTEXT("FontFaceReferencePreviewVisibility", "Reference"),
+		LOCTEXT("FontFaceReferencePreviewVisibilityTooltip", "Displays the Reference render of the preview text")
+	);
+	AddPreviewVisibilityItem(
+		EPreviewRow::ApproximateSdfLow,
+		LOCTEXT("FontFaceApproximateSdfLowPreviewVisibility", "Approximate SDF Low Quality"),
+		LOCTEXT("FontFaceApproximateSdfLowPreviewVisibilityTooltip", "Displays the preview text render of the fast approximation of the Low quality single-channel signed distance field")
+	);
+	AddPreviewVisibilityItem(
+		EPreviewRow::ApproximateSdfMedium,
+		LOCTEXT("FontFaceApproximateSdfMediumPreviewVisibility", "Approximate SDF Medium Quality"),
+		LOCTEXT("FontFaceApproximateSdfMediumPreviewVisibilityTooltip", "Displays the preview text render of the fast approximation of the Medium quality single-channel signed distance field")
+	);
+	AddPreviewVisibilityItem(
+		EPreviewRow::ApproximateSdfHigh,
+		LOCTEXT("FontFaceApproximateSdfHighPreviewVisibility", "Approximate SDF High Quality"),
+		LOCTEXT("FontFaceApproximateSdfHighPreviewVisibilityTooltip", "Displays the preview text render of the fast approximation of the High quality single-channel signed distance field")
+	);
+	AddPreviewVisibilityItem(
+		EPreviewRow::SdfLow,
+		LOCTEXT("FontFaceSdfLowPreviewVisibility", "SDF Low Quality"),
+		LOCTEXT("FontFaceSdfLowPreviewVisibilityTooltip", "Displays the Low quality single-channel signed distance field render of the preview text")
+	);
+	AddPreviewVisibilityItem(
+		EPreviewRow::SdfMedium,
+		LOCTEXT("FontFaceSdfMediumPreviewVisibility", "SDF Medium Quality"),
+		LOCTEXT("FontFaceSdfMediumPreviewVisibilityTooltip", "Displays the Medium quality single-channel signed distance field render of the preview text")
+	);
+	AddPreviewVisibilityItem(
+		EPreviewRow::SdfHigh,
+		LOCTEXT("FontFaceSdfHighPreviewVisibility", "SDF High Quality"),
+		LOCTEXT("FontFaceSdfHighPreviewVisibilityTooltip", "Displays the High quality single-channel signed distance field render of the preview text")
+	);
+	AddPreviewVisibilityItem(
+		EPreviewRow::MsdfLow,
+		LOCTEXT("FontFaceMsdfLowPreviewVisibility", "MSDF Low Quality"),
+		LOCTEXT("FontFaceMsdfLowPreviewVisibilityTooltip", "Displays the Low quality multi-channel signed distance field render of the preview text")
+	);
+	AddPreviewVisibilityItem(
+		EPreviewRow::MsdfMedium,
+		LOCTEXT("FontFaceMsdfMediumPreviewVisibility", "MSDF Medium Quality"),
+		LOCTEXT("FontFaceMsdfMediumPreviewVisibilityTooltip", "Displays the Medium quality multi-channel signed distance field render of the preview text")
+	);
+	AddPreviewVisibilityItem(
+		EPreviewRow::MsdfHigh,
+		LOCTEXT("FontFaceMsdfHighPreviewVisibility", "MSDF High Quality"),
+		LOCTEXT("FontFaceMsdfHighPreviewVisibilityTooltip", "Displays the High quality multi-channel signed distance field render of the preview text")
+	);
+
 	FontFacePreview =
 	SNew(SVerticalBox)
 	+SVerticalBox::Slot()
@@ -254,91 +334,128 @@ void FFontFaceEditor::CreateInternalWidgets()
 			[
 				SNew(SScrollBox)
 				.Orientation(EOrientation::Orient_Horizontal)
+				.ConsumeMouseWheel(EConsumeMouseWheel::Never)
 				+SScrollBox::Slot()
 				[
 					SAssignNew(PreviewTextGridPanel, SGridPanel)
-					+SGridPanel::Slot(0, 0)
+					+SGridPanel::Slot(0, (int32)EPreviewRow::Reference)
+					.VAlign(PreviewVAlign)
+					[
+						SAssignNew(PreviewTextBlocks[0][(int32)EPreviewRow::Reference], STextBlock)
+						.Text(LOCTEXT("FontFaceReferencePreviewLabel", "Reference: "))
+					]
+					+SGridPanel::Slot(1, (int32)EPreviewRow::Reference)
+					.VAlign(PreviewVAlign)
+					[
+						SAssignNew(PreviewTextBlocks[1][(int32)EPreviewRow::Reference], STextBlock)
+						.Text(DefaultPreviewText)
+					]
+					+SGridPanel::Slot(0, (int32)EPreviewRow::ApproximateSdfLow)
+					.VAlign(PreviewVAlign)
+					[
+						SAssignNew(PreviewTextBlocks[0][(int32)EPreviewRow::ApproximateSdfLow], STextBlock)
+						.Text(LOCTEXT("FontFaceApproximateSdfLowPreviewLabel", "ASDF Low: "))
+					]
+					+SGridPanel::Slot(1, (int32)EPreviewRow::ApproximateSdfLow)
+					.VAlign(PreviewVAlign)
+					[
+						SAssignNew(PreviewTextBlocks[1][(int32)EPreviewRow::ApproximateSdfLow], STextBlock)
+						.Text(DefaultPreviewText)
+					]
+					+SGridPanel::Slot(0, (int32)EPreviewRow::ApproximateSdfMedium)
 					.VAlign(PreviewVAlign)
 					[
-						SAssignNew(PreviewTextBlocks[0][0], STextBlock)
-						.Text(LOCTEXT("FontFaceReference", "Reference: "))
+						SAssignNew(PreviewTextBlocks[0][(int32)EPreviewRow::ApproximateSdfMedium], STextBlock)
+						.Text(LOCTEXT("FontFaceApproximateSdfMediumPreviewLabel", "ASDF Medium: "))
 					]
-					+SGridPanel::Slot(1, 0)
+					+SGridPanel::Slot(1, (int32)EPreviewRow::ApproximateSdfMedium)
 					.VAlign(PreviewVAlign)
 					[
-						SAssignNew(PreviewTextBlocks[1][0], STextBlock)
+						SAssignNew(PreviewTextBlocks[1][(int32)EPreviewRow::ApproximateSdfMedium], STextBlock)
 						.Text(DefaultPreviewText)
 					]
-					+SGridPanel::Slot(0, 1)
+					+SGridPanel::Slot(0, (int32)EPreviewRow::ApproximateSdfHigh)
 					.VAlign(PreviewVAlign)
 					[
-						SAssignNew(PreviewTextBlocks[0][1], STextBlock)
-						.Text(LOCTEXT("FontFaceLowQuality", "Low: "))
+						SAssignNew(PreviewTextBlocks[0][(int32)EPreviewRow::ApproximateSdfHigh], STextBlock)
+						.Text(LOCTEXT("FontFaceApproximateSdfHighPreviewLabel", "ASDF High: "))
 					]
-					+SGridPanel::Slot(1, 1)
+					+SGridPanel::Slot(1, (int32)EPreviewRow::ApproximateSdfHigh)
 					.VAlign(PreviewVAlign)
 					[
-						SAssignNew(PreviewTextBlocks[1][1], STextBlock)
+						SAssignNew(PreviewTextBlocks[1][(int32)EPreviewRow::ApproximateSdfHigh], STextBlock)
 						.Text(DefaultPreviewText)
 					]
-					+SGridPanel::Slot(0, 2)
+					+SGridPanel::Slot(0, (int32)EPreviewRow::SdfLow)
 					.VAlign(PreviewVAlign)
 					[
-						SAssignNew(PreviewTextBlocks[0][2], STextBlock)
-						.Text(LOCTEXT("FontFaceMediumQuality", "Medium: "))
+						SAssignNew(PreviewTextBlocks[0][(int32)EPreviewRow::SdfLow], STextBlock)
+						.Text(LOCTEXT("FontFaceSdfLowPreviewLabel", "SDF Low: "))
 					]
-					+SGridPanel::Slot(1, 2)
+					+SGridPanel::Slot(1, (int32)EPreviewRow::SdfLow)
 					.VAlign(PreviewVAlign)
 					[
-						SAssignNew(PreviewTextBlocks[1][2], STextBlock)
+						SAssignNew(PreviewTextBlocks[1][(int32)EPreviewRow::SdfLow], STextBlock)
 						.Text(DefaultPreviewText)
 					]
-					+SGridPanel::Slot(0, 3)
+					+SGridPanel::Slot(0, (int32)EPreviewRow::SdfMedium)
 					.VAlign(PreviewVAlign)
 					[
-						SAssignNew(PreviewTextBlocks[0][3], STextBlock)
-						.Text(LOCTEXT("FontFaceHighQuality", "High: "))
+						SAssignNew(PreviewTextBlocks[0][(int32)EPreviewRow::SdfMedium], STextBlock)
+						.Text(LOCTEXT("FontFaceSdfMediumPreviewLabel", "SDF Medium: "))
 					]
-					+SGridPanel::Slot(1, 3)
+					+SGridPanel::Slot(1, (int32)EPreviewRow::SdfMedium)
 					.VAlign(PreviewVAlign)
 					[
-						SAssignNew(PreviewTextBlocks[1][3], STextBlock)
+						SAssignNew(PreviewTextBlocks[1][(int32)EPreviewRow::SdfMedium], STextBlock)
 						.Text(DefaultPreviewText)
 					]
-					+SGridPanel::Slot(0, 4)
+					+SGridPanel::Slot(0, (int32)EPreviewRow::SdfHigh)
 					.VAlign(PreviewVAlign)
 					[
-						SAssignNew(PreviewTextBlocks[0][4], STextBlock)
-						.Text(LOCTEXT("FontFaceMultiLowQuality", "Multi Low: "))
+						SAssignNew(PreviewTextBlocks[0][(int32)EPreviewRow::SdfHigh], STextBlock)
+						.Text(LOCTEXT("FontFaceSdfHighPreviewLabel", "SDF High: "))
 					]
-					+SGridPanel::Slot(1, 4)
+					+SGridPanel::Slot(1, (int32)EPreviewRow::SdfHigh)
 					.VAlign(PreviewVAlign)
 					[
-						SAssignNew(PreviewTextBlocks[1][4], STextBlock)
+						SAssignNew(PreviewTextBlocks[1][(int32)EPreviewRow::SdfHigh], STextBlock)
 						.Text(DefaultPreviewText)
 					]
-					+SGridPanel::Slot(0, 5)
+					+SGridPanel::Slot(0, (int32)EPreviewRow::MsdfLow)
 					.VAlign(PreviewVAlign)
 					[
-						SAssignNew(PreviewTextBlocks[0][5], STextBlock)
-						.Text(LOCTEXT("FontFaceMultiMediumQuality", "Multi Medium: "))
+						SAssignNew(PreviewTextBlocks[0][(int32)EPreviewRow::MsdfLow], STextBlock)
+						.Text(LOCTEXT("FontFaceMsdfLowPreviewLabel", "MSDF Low: "))
 					]
-					+SGridPanel::Slot(1, 5)
+					+SGridPanel::Slot(1, (int32)EPreviewRow::MsdfLow)
 					.VAlign(PreviewVAlign)
 					[
-						SAssignNew(PreviewTextBlocks[1][5], STextBlock)
+						SAssignNew(PreviewTextBlocks[1][(int32)EPreviewRow::MsdfLow], STextBlock)
 						.Text(DefaultPreviewText)
 					]
-					+SGridPanel::Slot(0, 6)
+					+SGridPanel::Slot(0, (int32)EPreviewRow::MsdfMedium)
 					.VAlign(PreviewVAlign)
 					[
-						SAssignNew(PreviewTextBlocks[0][6], STextBlock)
-						.Text(LOCTEXT("FontFaceMultiHighQuality", "Multi High: "))
+						SAssignNew(PreviewTextBlocks[0][(int32)EPreviewRow::MsdfMedium], STextBlock)
+						.Text(LOCTEXT("FontFaceMsdfMediumPreviewLabel", "MSDF Medium: "))
 					]
-					+SGridPanel::Slot(1, 6)
+					+SGridPanel::Slot(1, (int32)EPreviewRow::MsdfMedium)
 					.VAlign(PreviewVAlign)
 					[
-						SAssignNew(PreviewTextBlocks[1][6], STextBlock)
+						SAssignNew(PreviewTextBlocks[1][(int32)EPreviewRow::MsdfMedium], STextBlock)
+						.Text(DefaultPreviewText)
+					]
+					+SGridPanel::Slot(0, (int32)EPreviewRow::MsdfHigh)
+					.VAlign(PreviewVAlign)
+					[
+						SAssignNew(PreviewTextBlocks[0][(int32)EPreviewRow::MsdfHigh], STextBlock)
+						.Text(LOCTEXT("FontFaceMsdfHighPreviewLabel", "MSDF High: "))
+					]
+					+SGridPanel::Slot(1, (int32)EPreviewRow::MsdfHigh)
+					.VAlign(PreviewVAlign)
+					[
+						SAssignNew(PreviewTextBlocks[1][(int32)EPreviewRow::MsdfHigh], STextBlock)
 						.Text(DefaultPreviewText)
 					]
 				]
@@ -365,7 +482,7 @@ void FFontFaceEditor::CreateInternalWidgets()
 			.OnTextChanged(this, &FFontFaceEditor::OnPreviewTextChanged)
 		]
 
-		+ SHorizontalBox::Slot()
+		+SHorizontalBox::Slot()
 		.AutoWidth()
 		[
 			SNew(SNumericEntryBox<int32>)
@@ -374,6 +491,25 @@ void FFontFaceEditor::CreateInternalWidgets()
 			.MaxValue(256)
 			.OnValueCommitted(this, &FFontFaceEditor::OnPreviewFontSizeChanged)
 		]
+
+		+SHorizontalBox::Slot()
+		.AutoWidth()
+		[
+			SAssignNew(PreviewVisibilityButton, SComboButton)
+			.HasDownArrow(false)
+			.ButtonStyle(FAppStyle::Get(), "SimpleButton")
+			.ToolTipText(LOCTEXT("FontFacePreviewVisibilityTooltip", "Selects which render modes to preview (requires Distance Field Rendering enabled)"))
+			.MenuContent()
+			[
+				PreviewRowVisibilitySelection.MakeWidget()
+			]
+			.ButtonContent()
+			[
+				SNew(SImage)
+				.Image(FAppStyle::Get().GetBrush("Level.VisibleIcon16x"))
+				.ColorAndOpacity(FSlateColor::UseForeground())
+			]
+		]
 	];
 
 	UpdatePreviewFonts();
@@ -494,30 +630,33 @@ void FFontFaceEditor::UpdatePreviewFonts()
 	{
 		return;
 	}
-	const int PreviewFontNum = UE_ARRAY_COUNT(*PreviewTextBlocks);
 	if (IsFontFaceDistanceFieldEnabled())
 	{
 		// This ensures that font geometry is preprocessed before cloning the face, otherwise it would be needlessly redone for each copy.
 		FontFace->CacheSubFaces();
-		PreviewFaces.SetNum(PreviewFontNum, EAllowShrinking::No);
-		PreviewFonts.SetNum(PreviewFontNum, EAllowShrinking::No);
-		ClonePreviewFontFace(PreviewFaces[0], EFontRasterizationMode::Bitmap);
-		ClonePreviewFontFace(PreviewFaces[1], EFontRasterizationMode::Sdf, FontFace->MinDistanceFieldPpem);
-		ClonePreviewFontFace(PreviewFaces[2], EFontRasterizationMode::Sdf, FontFace->MidDistanceFieldPpem);
-		ClonePreviewFontFace(PreviewFaces[3], EFontRasterizationMode::Sdf, FontFace->MaxDistanceFieldPpem);
-		ClonePreviewFontFace(PreviewFaces[4], EFontRasterizationMode::Msdf, FontFace->MinMultiDistanceFieldPpem);
-		ClonePreviewFontFace(PreviewFaces[5], EFontRasterizationMode::Msdf, FontFace->MidMultiDistanceFieldPpem);
-		ClonePreviewFontFace(PreviewFaces[6], EFontRasterizationMode::Msdf, FontFace->MaxMultiDistanceFieldPpem);
-		for (int32 Index = 0; Index < PreviewFontNum; ++Index)
+		PreviewFaces.SetNum((int32)EPreviewRow::Count, EAllowShrinking::No);
+		PreviewFonts.SetNum((int32)EPreviewRow::Count, EAllowShrinking::No);
+		ClonePreviewFontFace(PreviewFaces[(int32)EPreviewRow::Reference], EFontRasterizationMode::Bitmap);
+		ClonePreviewFontFace(PreviewFaces[(int32)EPreviewRow::ApproximateSdfLow], EFontRasterizationMode::SdfApproximation, FontFace->MinDistanceFieldPpem);
+		ClonePreviewFontFace(PreviewFaces[(int32)EPreviewRow::ApproximateSdfMedium], EFontRasterizationMode::SdfApproximation, FontFace->MidDistanceFieldPpem);
+		ClonePreviewFontFace(PreviewFaces[(int32)EPreviewRow::ApproximateSdfHigh], EFontRasterizationMode::SdfApproximation, FontFace->MaxDistanceFieldPpem);
+		ClonePreviewFontFace(PreviewFaces[(int32)EPreviewRow::SdfLow], EFontRasterizationMode::Sdf, FontFace->MinDistanceFieldPpem);
+		ClonePreviewFontFace(PreviewFaces[(int32)EPreviewRow::SdfMedium], EFontRasterizationMode::Sdf, FontFace->MidDistanceFieldPpem);
+		ClonePreviewFontFace(PreviewFaces[(int32)EPreviewRow::SdfHigh], EFontRasterizationMode::Sdf, FontFace->MaxDistanceFieldPpem);
+		ClonePreviewFontFace(PreviewFaces[(int32)EPreviewRow::MsdfLow], EFontRasterizationMode::Msdf, FontFace->MinMultiDistanceFieldPpem);
+		ClonePreviewFontFace(PreviewFaces[(int32)EPreviewRow::MsdfMedium], EFontRasterizationMode::Msdf, FontFace->MidMultiDistanceFieldPpem);
+		ClonePreviewFontFace(PreviewFaces[(int32)EPreviewRow::MsdfHigh], EFontRasterizationMode::Msdf, FontFace->MaxMultiDistanceFieldPpem);
+		for (int32 Index = 0; Index < (int32)EPreviewRow::Count; ++Index)
 		{
 			MakePreviewFont(PreviewFonts[Index], PreviewFaces[Index]);
 		}
 	}
 	else
 	{
+		static_assert((int32)EPreviewRow::Reference == 0, "Reference preview font needs to be the only array element (i.e. position 0)");
 		PreviewFaces.SetNum(1, EAllowShrinking::No);
-		PreviewFonts.SetNum(PreviewFontNum, EAllowShrinking::No);
-		ClonePreviewFontFace(PreviewFaces[0], EFontRasterizationMode::Bitmap);
+		PreviewFonts.SetNum((int32)EPreviewRow::Count, EAllowShrinking::No);
+		ClonePreviewFontFace(PreviewFaces[(int32)EPreviewRow::Reference], EFontRasterizationMode::Bitmap);
 		for (TObjectPtr<UObject>& PreviewFont : PreviewFonts)
 		{
 			MakePreviewFont(PreviewFont, PreviewFaces[0]);
@@ -529,12 +668,15 @@ void FFontFaceEditor::UpdatePreviewVisibility()
 {
 	if (FontFace)
 	{
-		const EVisibility SecondaryRowsVisibility = IsFontFaceDistanceFieldEnabled() ? EVisibility::Visible : EVisibility::Collapsed;
-		PreviewTextBlocks[0][0]->SetVisibility(SecondaryRowsVisibility);
-		for (int32 Index = 1; Index < UE_ARRAY_COUNT(*PreviewTextBlocks); ++Index)
+		const bool bSecondaryRowsVisibility = IsFontFaceDistanceFieldEnabled();
+		PreviewTextBlocks[0][(int32)EPreviewRow::Reference]->SetVisibility(bSecondaryRowsVisibility && PreviewRowVisibility[(int32)EPreviewRow::Reference] ? EVisibility::Visible : EVisibility::Collapsed);
+		PreviewTextBlocks[1][(int32)EPreviewRow::Reference]->SetVisibility(!bSecondaryRowsVisibility || PreviewRowVisibility[(int32)EPreviewRow::Reference] ? EVisibility::Visible : EVisibility::Collapsed);
+		static_assert((int32)EPreviewRow::Reference == 0, "Visibility not set for rows lower than EPreviewRow::Reference");
+		for (int32 PreviewRow = (int32)EPreviewRow::Reference + 1; PreviewRow < (int32)EPreviewRow::Count; ++PreviewRow)
 		{
-			PreviewTextBlocks[0][Index]->SetVisibility(SecondaryRowsVisibility);
-			PreviewTextBlocks[1][Index]->SetVisibility(SecondaryRowsVisibility);
+			const EVisibility RowVisibility = bSecondaryRowsVisibility && PreviewRowVisibility[PreviewRow] ? EVisibility::Visible : EVisibility::Collapsed;
+			PreviewTextBlocks[0][PreviewRow]->SetVisibility(RowVisibility);
+			PreviewTextBlocks[1][PreviewRow]->SetVisibility(RowVisibility);
 		}
 		PreviewNoteTextBlock->SetVisibility(
 			FontFace->bEnableDistanceFieldRendering &&
@@ -542,6 +684,16 @@ void FFontFaceEditor::UpdatePreviewVisibility()
 			!GetDefault<UUserInterfaceSettings>()->bEnableDistanceFieldFontRasterization ?
 			EVisibility::Visible : EVisibility::Collapsed
 		);
+		PreviewVisibilityButton->SetEnabled(bSecondaryRowsVisibility);
+	}
+	else
+	{
+		for (int32 PreviewRow = 0; PreviewRow < (int32)EPreviewRow::Count; ++PreviewRow)
+		{
+			PreviewTextBlocks[0][PreviewRow]->SetVisibility(EVisibility::Collapsed);
+			PreviewTextBlocks[1][PreviewRow]->SetVisibility(EVisibility::Collapsed);
+		}
+		PreviewVisibilityButton->SetEnabled(false);
 	}
 }
 
@@ -572,6 +724,21 @@ void FFontFaceEditor::ApplyPreviewFontSize()
 	}
 }
 
+void FFontFaceEditor::ChangePreviewRowVisibility(EPreviewRow Row)
+{
+	int32 RowIndex = (int32)Row;
+	check(RowIndex >= 0 && RowIndex < (int32)EPreviewRow::Count);
+	PreviewRowVisibility[RowIndex] = !PreviewRowVisibility[RowIndex];
+	UpdatePreviewVisibility();
+}
+
+bool FFontFaceEditor::GetPreviewRowVisibility(EPreviewRow Row) const
+{
+	int32 RowIndex = (int32)Row;
+	check(RowIndex >= 0 && RowIndex < (int32)EPreviewRow::Count);
+	return PreviewRowVisibility[RowIndex];
+}
+
 float FFontFaceEditor::GetPreviewTextPadding() const
 {
 	return PreviewFontSize / 2.0f;
diff --git a/Engine/Source/Editor/FontEditor/Private/FontFaceEditor.h b/Engine/Source/Editor/FontEditor/Private/FontFaceEditor.h
--- a/Engine/Source/Editor/FontEditor/Private/FontFaceEditor.h
+++ b/Engine/Source/Editor/FontEditor/Private/FontFaceEditor.h
@@ -18,6 +18,7 @@
 #include "UObject/Object.h"
 
 class FReferenceCollector;
+class SComboButton;
 class SCompositeFontEditor;
 class SDockTab;
 class SEditableTextBox;
@@ -98,6 +99,23 @@ private:
 	bool GetIsPropertyVisible(const struct FPropertyAndParent& PropertyAndParent) const;
 
 private:
+
+	/** Preview text rows that can have their visibility toggled */
+	enum class EPreviewRow : int32
+	{
+		Reference,
+		ApproximateSdfLow,
+		ApproximateSdfMedium,
+		ApproximateSdfHigh,
+		SdfLow,
+		SdfMedium,
+		SdfHigh,
+		MsdfLow,
+		MsdfMedium,
+		MsdfHigh,
+		Count
+	};
+
 	/** The font asset being inspected */
 	TObjectPtr<UFontFace> FontFace;
 
@@ -119,8 +137,11 @@ private:
 	/** Preview text */
 	TSharedPtr<SEditableTextBox> FontFacePreviewText;
 
+	/** Preview visibility menu button */
+	TSharedPtr<SComboButton> PreviewVisibilityButton;
+
 	/** Preview widgets */
-	TSharedPtr<STextBlock> PreviewTextBlocks[2][7];
+	TSharedPtr<STextBlock> PreviewTextBlocks[2][(int32)EPreviewRow::Count];
 
 	/** Preview grid panel */
 	TSharedPtr<SGridPanel> PreviewTextGridPanel;
@@ -131,6 +152,9 @@ private:
 	/** Preview font size */
 	int32 PreviewFontSize = 24;
 
+	/** Preview row user visibility filter */
+	bool PreviewRowVisibility[(int32)EPreviewRow::Count];
+
 	/** The tab ids for the font editor */
 	static const FName PreviewTabId;
 	static const FName PropertiesTabId;
@@ -142,5 +166,7 @@ private:
 	void UpdatePreviewFonts();
 	void UpdatePreviewVisibility();
 	void ApplyPreviewFontSize();
+	void ChangePreviewRowVisibility(EPreviewRow Row);
+	bool GetPreviewRowVisibility(EPreviewRow Row) const;
 	float GetPreviewTextPadding() const;
 };
