Remove deprecated InputAxis and InputKey function variations on the FViewportClient from 5.1

diff --git a/Engine/Source/Editor/DistCurveEditor/Private/CurveEditorViewportClient.cpp b/Engine/Source/Editor/DistCurveEditor/Private/CurveEditorViewportClient.cpp
--- a/Engine/Source/Editor/DistCurveEditor/Private/CurveEditorViewportClient.cpp
+++ b/Engine/Source/Editor/DistCurveEditor/Private/CurveEditorViewportClient.cpp
@@ -288,10 +288,14 @@ void FCurveEditorViewportClient::Draw(FViewport* Viewport, FCanvas* Canvas)
 	}
 }
 
-bool FCurveEditorViewportClient::InputKey(FViewport* Viewport, int32 ControllerId, FKey Key, EInputEvent Event, float AmountDepressed, bool Gamepad)
+bool FCurveEditorViewportClient::InputKey(const FInputKeyEventArgs& EventArgs)
 {	
 	bool bHandled = false;
 
+	FViewport* Viewport = EventArgs.Viewport;
+	const FKey& Key = EventArgs.Key;
+	const EInputEvent Event = EventArgs.Event;
+
 	bool bCtrlDown = Viewport->KeyState(EKeys::LeftControl) || Viewport->KeyState(EKeys::RightControl);
 	bool bShiftDown = Viewport->KeyState(EKeys::LeftShift) || Viewport->KeyState(EKeys::RightShift);
 	bool bAltDown = Viewport->KeyState(EKeys::LeftAlt) || Viewport->KeyState(EKeys::RightAlt);
@@ -801,7 +805,7 @@ void FCurveEditorViewportClient::CapturedMouseMove(FViewport* Viewport, int32 X,
 	MouseMove(Viewport, X, Y);
 }
 
-bool FCurveEditorViewportClient::InputAxis(FViewport* Viewport, int32 ControllerId, FKey Key, float Delta, float DeltaTime, int32 NumSamples, bool bGamepad)
+bool FCurveEditorViewportClient::InputAxis(FViewport* Viewport, FInputDeviceId InputDevice, FKey Key, float Delta, float DeltaTime, int32 NumSamples , bool bGamepad)
 {
 	if (Key == EKeys::MouseX || Key == EKeys::MouseY)
 	{
diff --git a/Engine/Source/Editor/DistCurveEditor/Private/CurveEditorViewportClient.h b/Engine/Source/Editor/DistCurveEditor/Private/CurveEditorViewportClient.h
--- a/Engine/Source/Editor/DistCurveEditor/Private/CurveEditorViewportClient.h
+++ b/Engine/Source/Editor/DistCurveEditor/Private/CurveEditorViewportClient.h
@@ -31,14 +31,14 @@ class FCurveEditorViewportClient : public FCommonViewportClient
 public:
 	/** Constructor */
 	FCurveEditorViewportClient(TWeakPtr<SDistributionCurveEditor> InCurveEditor, TWeakPtr<SCurveEditorViewport> InCurveEditorViewport);
-	~FCurveEditorViewportClient();
+	virtual ~FCurveEditorViewportClient() override;
 
 	/** FViewportClient interface */
 	virtual void Draw(FViewport* Viewport, FCanvas* Canvas) override;
-	virtual bool InputKey(FViewport* Viewport, int32 ControllerId, FKey Key, EInputEvent Event, float AmountDepressed = 1.0f, bool bGamepad = false) override;
+	virtual bool InputKey(const FInputKeyEventArgs& EventArgs) override;
 	virtual void MouseMove(FViewport* Viewport, int32 X, int32 Y) override;
 	virtual void CapturedMouseMove(FViewport* Viewport, int32 X, int32 Y) override;
-	virtual bool InputAxis(FViewport* Viewport, int32 ControllerId, FKey Key, float Delta, float DeltaTime, int32 NumSamples = 1, bool bGamepad = false) override;
+	virtual bool InputAxis(FViewport* Viewport, FInputDeviceId InputDevice, FKey Key, float Delta, float DeltaTime, int32 NumSamples = 1, bool bGamepad = false) override;
 
 	/** Exec handler */
 	void Exec(const TCHAR* Cmd);
diff --git a/Engine/Source/Editor/FontEditor/Private/SFontEditorViewport.cpp b/Engine/Source/Editor/FontEditor/Private/SFontEditorViewport.cpp
--- a/Engine/Source/Editor/FontEditor/Private/SFontEditorViewport.cpp
+++ b/Engine/Source/Editor/FontEditor/Private/SFontEditorViewport.cpp
@@ -67,11 +67,11 @@ class FFontEditorViewportClient : public FCommonViewportClient
 public:
 	/** Constructor */
 	FFontEditorViewportClient(TWeakPtr<SFontEditorViewport> InFontEditorViewport);
-	~FFontEditorViewportClient();
+	virtual ~FFontEditorViewportClient() override;
 
 	/** FLevelEditorViewportClient interface */
 	virtual void Draw(FViewport* Viewport, FCanvas* Canvas) override;
-	virtual bool InputKey(FViewport* Viewport, int32 ControllerId, FKey Key, EInputEvent Event, float AmountDepressed = 1.f, bool bGamepad = false) override;
+	virtual bool InputKey(const FInputKeyEventArgs& EventArgs) override;
 	
 	/** Determines which texture page was selected. */
 	void UpdateSelectedPage(UObject* SelectedObject);
@@ -437,7 +437,7 @@ void FFontEditorViewportClient::Draw(FViewport* Viewport, FCanvas* Canvas)
 	}
 }
 
-bool FFontEditorViewportClient::InputKey(FViewport* Viewport, int32 ControllerId, FKey Key, EInputEvent Event, float AmountDepressed, bool Gamepad)
+bool FFontEditorViewportClient::InputKey(const FInputKeyEventArgs& EventArgs)
 {
 	// The preview viewport doesn't need to process key input
 	if (FontEditorViewportPtr.Pin()->IsPreviewViewport())
@@ -447,14 +447,17 @@ bool FFontEditorViewportClient::InputKey(FViewport* Viewport, int32 ControllerId
 
 	bool bHandled = false;
 
-	if (Key == EKeys::LeftMouseButton && Event == IE_Released)
+	const FKey& Key = EventArgs.Key;
+	FViewport* Viewport = EventArgs.Viewport;
+
+	if (Key == EKeys::LeftMouseButton && EventArgs.Event == IE_Released)
 	{
-		Viewport->Invalidate();
+		EventArgs.Viewport->Invalidate();
 		bHandled = true;
 	}
 	else
 	{
-		if((Key == EKeys::LeftMouseButton || Key == EKeys::RightMouseButton) && Event == IE_Pressed)
+		if((Key == EKeys::LeftMouseButton || Key == EKeys::RightMouseButton) && EventArgs.Event == IE_Pressed)
 		{
 			const int32 HitX = Viewport->GetMouseX();
 			const int32 HitY = Viewport->GetMouseY();
diff --git a/Engine/Source/Runtime/Engine/Private/UnrealClient.cpp b/Engine/Source/Runtime/Engine/Private/UnrealClient.cpp
--- a/Engine/Source/Runtime/Engine/Private/UnrealClient.cpp
+++ b/Engine/Source/Runtime/Engine/Private/UnrealClient.cpp
@@ -2539,13 +2539,6 @@ void FViewportClient::RequestInvalidateHitProxy(FViewport* Viewport)
 	Viewport->InvalidateHitProxy();
 }
 
-bool FViewportClient::InputKey(const FInputKeyEventArgs& EventArgs)
-{
-	PRAGMA_DISABLE_DEPRECATION_WARNINGS
-	return InputKey(EventArgs.Viewport, EventArgs.ControllerId, EventArgs.Key, EventArgs.Event, EventArgs.AmountDepressed, EventArgs.Key.IsGamepadKey());
-	PRAGMA_ENABLE_DEPRECATION_WARNINGS
-}
-
 TOptional<TSharedRef<SWidget>> FViewportClient::MapCursor(FViewport* Viewport, const FCursorReply& CursorReply)
 {
 	return TOptional<TSharedRef<SWidget>>();
diff --git a/Engine/Source/Runtime/Engine/Public/ViewportClient.h b/Engine/Source/Runtime/Engine/Public/ViewportClient.h
--- a/Engine/Source/Runtime/Engine/Public/ViewportClient.h
+++ b/Engine/Source/Runtime/Engine/Public/ViewportClient.h
@@ -42,25 +42,7 @@ public:
 	 * @param	EventArgs - The Input event args.
 	 * @return	True to consume the key event, false to pass it on.
 	 */
-	ENGINE_API virtual bool InputKey(const FInputKeyEventArgs& EventArgs);
-
-	UE_DEPRECATED(5.1, "Use the new InputKey(const FInputKeyEventArgs& EventArgs) function.")
-	virtual bool InputKey(FViewport* Viewport,int32 ControllerId,FKey Key,EInputEvent Event,float AmountDepressed = 1.f,bool bGamepad=false) { return false; }
-
-	/**
-	 * Check an axis movement received by the viewport.
-	 * If the viewport client uses the movement, it should return true to consume it.
-	 * @param	Viewport - The viewport which the axis movement is from.
-	 * @param	ControllerId - The controller which the axis movement is from.
-	 * @param	Key - The name of the axis which moved.
-	 * @param	Delta - The axis movement delta.
-	 * @param	DeltaTime - The time since the last axis update.
-	 * @param	NumSamples - The number of device samples that contributed to this Delta, useful for things like smoothing
-	 * @param	bGamepad - input came from gamepad (ie xbox controller)
-	 * @return	True to consume the axis movement, false to pass it on.
-	 */
-	UE_DEPRECATED(5.1, "Use the new InputAxis function that takes a FInputDeviceId.")
-	virtual bool InputAxis(FViewport* Viewport,int32 ControllerId,FKey Key,float Delta,float DeltaTime,int32 NumSamples=1,bool bGamepad=false) { return false; }
+	ENGINE_API virtual bool InputKey(const FInputKeyEventArgs& EventArgs) { return false; }
 
 	/**
 	 * Check an axis movement received by the viewport.
