Support export of font files (ttf,...) from the font face, with a new export button next to the import one.

#tests Exported a few existing fonts, opened them in font viewer, reimport them as new font.
#FailedValidations Preflight

diff --git a/Engine/Source/Editor/FontEditor/Private/FontFaceDetailsCustomization.cpp b/Engine/Source/Editor/FontEditor/Private/FontFaceDetailsCustomization.cpp
--- a/Engine/Source/Editor/FontEditor/Private/FontFaceDetailsCustomization.cpp
+++ b/Engine/Source/Editor/FontEditor/Private/FontFaceDetailsCustomization.cpp
@@ -42,6 +42,10 @@
 
 #define LOCTEXT_NAMESPACE "FontFaceDetailsCustomization"
 
+namespace Private
+{
+	static const FString FileTypes = TEXT("All Font Files (*.ttf, *.ttc, *.otf, *.otc)|*.ttf;*.ttc;*.otf;*.otc|TrueType fonts (*.ttf, *.ttc)|*.ttf;*.ttc|OpenType fonts (*.otf, *.otc)|*.otf;*.otc");
+}
 void FFontFaceDetailsCustomization::CustomizeDetails(IDetailLayoutBuilder& DetailBuilder)
 {
 	DetailBuilder.GetObjectsBeingCustomized(ObjectsBeingEdited);
@@ -92,6 +96,24 @@ void FFontFaceDetailsCustomization::CustomizeDetails(IDetailLayoutBuilder& Detai
 						.Text(FEditorFontGlyphs::Folder_Open)
 					]
 				]
+				+ SHorizontalBox::Slot()
+					.AutoWidth()
+					.Padding(4.0f, 0.0f, 0.0f, 0.0f)
+					.VAlign(VAlign_Center)
+					[
+						SNew(SButton)
+							.ButtonStyle(FAppStyle::Get(), "HoverHintOnly")
+							.ToolTipText(LOCTEXT("FontFileExportFontToolTip", "Export the font file to this computer"))
+							.OnClicked(this, &FFontFaceDetailsCustomization::OnExportFont)
+							.ContentPadding(2.0f)
+							.ForegroundColor(FSlateColor::UseForeground())
+							.IsFocusable(false)
+							[
+								SNew(STextBlock)
+									.Font(FAppStyle::Get().GetFontStyle("FontAwesome.10"))
+									.Text(FEditorFontGlyphs::File)
+							]
+					]
 			];
 	}
 
@@ -180,7 +202,7 @@ FReply FFontFaceDetailsCustomization::OnBrowseFontPath()
 			LOCTEXT("FontPickerTitle", "Choose a font file...").ToString(),
 			DefaultPath,
 			TEXT(""),
-			TEXT("All Font Files (*.ttf, *.ttc, *.otf, *.otc)|*.ttf;*.ttc;*.otf;*.otc|TrueType fonts (*.ttf, *.ttc)|*.ttf;*.ttc|OpenType fonts (*.otf, *.otc)|*.otf;*.otc"),
+			Private::FileTypes,
 			EFileDialogFlags::None,
 			OutFiles
 			))
@@ -192,6 +214,43 @@ FReply FFontFaceDetailsCustomization::OnBrowseFontPath()
 	return FReply::Handled();
 }
 
+FReply FFontFaceDetailsCustomization::OnExportFont()
+{
+	IDesktopPlatform* const DesktopPlatform = FDesktopPlatformModule::Get();
+	if (DesktopPlatform)
+	{
+		const void* ParentWindowWindowHandle = FSlateApplication::Get().FindBestParentWindowHandleForDialogs(nullptr);
+
+		for (const TWeakObjectPtr<UObject>& ObjectBeingEdited : ObjectsBeingEdited)
+		{
+			UFontFace* FontFace = Cast<UFontFace>(ObjectBeingEdited.Get());
+			if (FontFace)
+			{
+				const FString CurrentFilename = FPaths::GetBaseFilename(FontFace->SourceFilename);
+				const FText Title = FText::Format(LOCTEXT("FontFace_ExportFontFileDialogTitle", "Export font '{0}'..."), FText::FromString(CurrentFilename));
+
+				TArray<FString> OutFilenames;
+				DesktopPlatform->SaveFileDialog(
+					ParentWindowWindowHandle,
+					Title.ToString(),
+					(FontFace->SourceFilename.IsEmpty()) ? TEXT("") : FPaths::GetPath(FontFace->SourceFilename),
+					(CurrentFilename.IsEmpty()) ? TEXT("") : CurrentFilename + TEXT(".ttf"),
+					Private::FileTypes,
+					EFileDialogFlags::None,
+					OutFilenames
+				);
+
+				if (!OutFilenames.IsEmpty())
+				{
+					FFileHelper::SaveArrayToFile(FontFace->FontFaceData.Get().GetData(), *OutFilenames[0]);
+				}
+			}
+		}
+	}
+
+	return FReply::Handled();
+}
+
 void FFontFaceDetailsCustomization::OnFontPathPicked(const FString& InNewFontFilename)
 {
 	TArray<uint8> FontData;
diff --git a/Engine/Source/Editor/FontEditor/Private/FontFaceDetailsCustomization.h b/Engine/Source/Editor/FontEditor/Private/FontFaceDetailsCustomization.h
--- a/Engine/Source/Editor/FontEditor/Private/FontFaceDetailsCustomization.h
+++ b/Engine/Source/Editor/FontEditor/Private/FontFaceDetailsCustomization.h
@@ -37,6 +37,9 @@ private:
 	/** Called in response to the user wanting to pick a new font file. */
 	FReply OnBrowseFontPath();
 
+	/** Called in response to the user wanting to export the font to a file. */
+	FReply OnExportFont();
+
 	/** Called in response to a new font file being picked. */
 	void OnFontPathPicked(const FString& InNewFontFilename);
 
